Enabled warnings that alert when asset registry queries during cook search for in-memory assets

#Synced-CL 52508799

diff --git a/Engine/Source/Editor/UnrealEd/Private/Cooker/CookAssetRegistryAccessTracker.cpp b/Engine/Source/Editor/UnrealEd/Private/Cooker/CookAssetRegistryAccessTracker.cpp
--- a/Engine/Source/Editor/UnrealEd/Private/Cooker/CookAssetRegistryAccessTracker.cpp
+++ b/Engine/Source/Editor/UnrealEd/Private/Cooker/CookAssetRegistryAccessTracker.cpp
@@ -173,13 +173,12 @@ namespace UE::CookAssetRegistryAccessTracker
 			NewRecord.Filter = UE::AssetRegistry::ConvertToNonCompiledFilter(CompiledFilter);
 			if (!NewRecord.Filter.bIncludeOnlyOnDiskAssets)
 			{
-				// TODO: Bump this log severity up to warning and enable DumpStackTraceToLog once we have fixed all engine classes.
-				UE_LOG(LogCook, Verbose,
+				UE_LOG(LogCook, Warning,
 					TEXT("An AssetRegistry query is executing during a package load or save, and it uses an ARFilter with bIncludeOnlyOnDiskAssets=false. ")
 					TEXT("This will create cookindeterminism since the packages in can vary between cooks. Caller should set bIncludeOnlyOnDiskAssets=true. ")
 					TEXT("For the recorded dependency we are forcing bIncludeOnlyOnDiskAssets=true. Callstack of the AssetRegistryQuery construction: ")
 				);
-				//FDebug::DumpStackTraceToLog(ELogVerbosity::Display);
+				FDebug::DumpStackTraceToLog(ELogVerbosity::Display);
 				NewRecord.Filter.bIncludeOnlyOnDiskAssets = true;
 			}
 			NewRecord.Filter.SortForSaving();
diff --git a/Engine/Source/Runtime/CoreUObject/Private/Cooker/CookDependency.cpp b/Engine/Source/Runtime/CoreUObject/Private/Cooker/CookDependency.cpp
--- a/Engine/Source/Runtime/CoreUObject/Private/Cooker/CookDependency.cpp
+++ b/Engine/Source/Runtime/CoreUObject/Private/Cooker/CookDependency.cpp
@@ -230,13 +230,12 @@ FCookDependency FCookDependency::AssetRegistryQuery(FARFilter Filter)
 	Filter.SortForSaving();
 	if (!Filter.bIncludeOnlyOnDiskAssets)
 	{
-		// TODO: Bump this log severity up to warning and enable DumpStackTraceToLog once we have fixed all engine classes.
-		UE_LOG(LogCore, Verbose,
+		UE_LOG(LogCore, Warning,
 			TEXT("An AssetRegistryQuery dependency has an ARFilter with bIncludeOnlyOnDiskAssets=false. ")
 			TEXT("This will create cookindeterminism since the packages in can vary between cooks. Caller should set bIncludeOnlyOnDiskAssets=true. ")
 			TEXT("For the recorded dependency we are forcing bIncludeOnlyOnDiskAssets=true. Callstack of the AssetRegistryQuery construction: ")
 		);
-		//FDebug::DumpStackTraceToLog(ELogVerbosity::Display);
+		FDebug::DumpStackTraceToLog(ELogVerbosity::Display);
 		Filter.bIncludeOnlyOnDiskAssets = true;
 	}
 
