- Set bCommentBubbleVisible when comments are added
- This fixes the comments not persisting between material editor close / opens

#Synced-CL 52445503

diff --git a/Engine/Source/Editor/UnrealEd/Private/MaterialGraphNode.cpp b/Engine/Source/Editor/UnrealEd/Private/MaterialGraphNode.cpp
--- a/Engine/Source/Editor/UnrealEd/Private/MaterialGraphNode.cpp
+++ b/Engine/Source/Editor/UnrealEd/Private/MaterialGraphNode.cpp
@@ -120,6 +120,7 @@ void UMaterialGraphNode::PostEditChangeProperty(FPropertyChangedEvent& PropertyC
 				if (MaterialExpression->Desc != NodeComment)
 				{
 					MaterialExpression->Desc = NodeComment;
+					MaterialExpression->bCommentBubbleVisible = !NodeComment.IsEmpty();
 					MaterialExpression->Material->PropagateExpressionParameterChanges(MaterialExpression);
 				}
 			}
@@ -952,6 +953,7 @@ void UMaterialGraphNode::OnUpdateCommentText( const FString& NewComment )
 		if (MaterialExpression->Desc != NewComment)
 		{
 			MaterialExpression->Desc = NewComment;
+			MaterialExpression->bCommentBubbleVisible = !NewComment.IsEmpty();
 			MaterialExpression->Material->PropagateExpressionParameterChanges(MaterialExpression);
 		}
 		MaterialDirtyDelegate.ExecuteIfBound();
