[TRS Gizmo] Do not update non-realtime viewports when manipulating

#Synced-CL 52496820

diff --git a/Engine/Source/Editor/UnrealEd/Private/Tools/EdModeInteractiveToolsContext.cpp b/Engine/Source/Editor/UnrealEd/Private/Tools/EdModeInteractiveToolsContext.cpp
--- a/Engine/Source/Editor/UnrealEd/Private/Tools/EdModeInteractiveToolsContext.cpp
+++ b/Engine/Source/Editor/UnrealEd/Private/Tools/EdModeInteractiveToolsContext.cpp
@@ -498,17 +498,17 @@ UWorld* UEditorInteractiveToolsContext::GetWorld() const
 
 void UEditorInteractiveToolsContext::Tick(FEditorViewportClient* ViewportClient, float DeltaTime)
 {
-	// invalidate this viewport if it's timestamp is not current
-	const int32* FoundTimestamp = InvalidationMap.Find(ViewportClient);
-	if (ViewportClient && FoundTimestamp == nullptr)
+	// Invalidate this viewport if its timestamp is not current.
+	// During active tracking (e.g. gizmo drag), only invalidate the focused viewport
+	const bool bSkipInvalidation = EditorModeManager && EditorModeManager->IsTracking() && ViewportClient != EditorModeManager->GetFocusedViewportClient();
+	if (!bSkipInvalidation && ViewportClient)
 	{
-		ViewportClient->Invalidate(false, false);
-		InvalidationMap.Add(ViewportClient, InvalidationTimestamp);
-	}
-	if (ViewportClient && FoundTimestamp != nullptr && *FoundTimestamp < InvalidationTimestamp)
-	{
-		ViewportClient->Invalidate(false, false);
-		InvalidationMap[ViewportClient] = InvalidationTimestamp;
+		int32& ViewportInvalidation = InvalidationMap.FindOrAdd(ViewportClient);
+		if(ViewportInvalidation < InvalidationTimestamp)
+		{
+			ViewportClient->Invalidate(false, false);
+			ViewportInvalidation = InvalidationTimestamp;
+		}
 	}
 
 	if ( PendingToolShutdownType )
