[TRS Gizmo] Fix C6011 null-pointer warning in UEditorInteractiveToolsContext::Tick

#Synced-CL 52496820

diff --git a/Engine/Source/Editor/UnrealEd/Private/Tools/EdModeInteractiveToolsContext.cpp b/Engine/Source/Editor/UnrealEd/Private/Tools/EdModeInteractiveToolsContext.cpp
--- a/Engine/Source/Editor/UnrealEd/Private/Tools/EdModeInteractiveToolsContext.cpp
+++ b/Engine/Source/Editor/UnrealEd/Private/Tools/EdModeInteractiveToolsContext.cpp
@@ -527,7 +527,7 @@ void UEditorInteractiveToolsContext::Tick(FEditorViewportClient* ViewportClient,
 	
 	// This Tick() is called for every ViewportClient, however we only want to Tick the ToolManager and GizmoManager
 	// once, for the 'Active'/Focused Viewport, so early-out here
-	if (ViewportClient != EditorModeManager->GetFocusedViewportClient())
+	if (EditorModeManager && EditorModeManager->GetFocusedViewportClient() != ViewportClient)
 	{
 		return;
 	}
