Fix localization warnings

#Synced-CL 52520240

diff --git a/Engine/Plugins/Runtime/StateTree/Source/StateTreeEditorModule/Private/SStateTreeViewRow.cpp b/Engine/Plugins/Runtime/StateTree/Source/StateTreeEditorModule/Private/SStateTreeViewRow.cpp
--- a/Engine/Plugins/Runtime/StateTree/Source/StateTreeEditorModule/Private/SStateTreeViewRow.cpp
+++ b/Engine/Plugins/Runtime/StateTree/Source/StateTreeEditorModule/Private/SStateTreeViewRow.cpp
@@ -1255,7 +1255,7 @@ void SStateTreeViewRow::MakeFlagsWidget()
 						SNew(SImage)
 						.Image(FStateTreeEditorStyle::Get().GetBrush("StateTreeEditor.Flags.TickOnEvent"))
 						.ColorAndOpacity(IconTint)
-						.ToolTipText(LOCTEXT("StateNodeTickTransition", "The state contains at least one transition that ticks at runtime when there's an event."))
+						.ToolTipText(LOCTEXT("StateNodeTickTransitionEvent", "The state contains at least one transition that ticks at runtime when there's an event."))
 					];
 				}
 				else
diff --git a/Engine/Source/Editor/MirrorDataTableEditor/Private/MirrorDataTableEditorModule.cpp b/Engine/Source/Editor/MirrorDataTableEditor/Private/MirrorDataTableEditorModule.cpp
--- a/Engine/Source/Editor/MirrorDataTableEditor/Private/MirrorDataTableEditorModule.cpp
+++ b/Engine/Source/Editor/MirrorDataTableEditor/Private/MirrorDataTableEditorModule.cpp
@@ -243,7 +243,7 @@ void FMirrorDataTableEditorModule::RegisterMenus()
 									OptionsSection.AddMenuEntry(
 										NAME_None,
 										LOCTEXT("MapToSelfLabel", "Default to self on find/replace failure"),
-										LOCTEXT("MapToSelfLabel", "Maps entry name to itself when the find/replace fails to find a mirror entry name."),
+										LOCTEXT("MapToSelfTooltip", "Maps entry name to itself when the find/replace fails to find a mirror entry name."),
 										FSlateIcon(),
 										FUIAction(
 											FExecuteAction::CreateLambda([]()
diff --git a/Engine/Source/Editor/UnrealEd/Private/AutoReimport/ContentDirectoryMonitor.cpp b/Engine/Source/Editor/UnrealEd/Private/AutoReimport/ContentDirectoryMonitor.cpp
--- a/Engine/Source/Editor/UnrealEd/Private/AutoReimport/ContentDirectoryMonitor.cpp
+++ b/Engine/Source/Editor/UnrealEd/Private/AutoReimport/ContentDirectoryMonitor.cpp
@@ -321,7 +321,7 @@ void FContentDirectoryMonitor::ProcessAdditions(const DirectoryWatcher::FTimeLim
 					if (!bImportHandledThroughInterchange)
 					{
 						// Add the message to the message log letting the user know that the import didn't work.
-						Context.GetMessageLog().Message(EMessageSeverity::Info, FText::Format(LOCTEXT("Info_FailedToImportAsset", "Interchange: Failed to import file {0}."), FText::FromString(FullFilename)));
+						Context.GetMessageLog().Message(EMessageSeverity::Info, FText::Format(LOCTEXT("Info_InterchangeFailedToImportAsset", "Interchange: Failed to import file {0}."), FText::FromString(FullFilename)));
 					}
 				}
 			}
