[ITF Gizmos] Level Editor: making sure selection affected by transforms makes it into undo transaction history

#Synced-CL 52498372

diff --git a/Engine/Source/Editor/Experimental/EditorInteractiveToolsFramework/Private/EditorGizmos/EditorTRSGizmo.cpp b/Engine/Source/Editor/Experimental/EditorInteractiveToolsFramework/Private/EditorGizmos/EditorTRSGizmo.cpp
--- a/Engine/Source/Editor/Experimental/EditorInteractiveToolsFramework/Private/EditorGizmos/EditorTRSGizmo.cpp
+++ b/Engine/Source/Editor/Experimental/EditorInteractiveToolsFramework/Private/EditorGizmos/EditorTRSGizmo.cpp
@@ -4964,16 +4964,6 @@ void UEditorTRSGizmo::OnKeyPressed(const FKey& KeyID)
 		if (NudgeData.PressedKeys.IsEmpty())
 		{
 			BeginTransformEditSequence();
-
-			if (TStrongObjectPtr<const UTypedElementSelectionSet> SelectionSet = SelectionSetWeak.Pin())
-			{
-				TArray<AActor*> SelectedActors = SelectionSet->GetSelectedObjects<AActor>();
-				for (AActor* const Actor : SelectedActors)
-				{
-					// Modify is only called for mouse interactions, so we need to handle it directly
-					Actor->Modify();
-				}
-			}
 		}
 
 		NudgeData.PressedKeys.AddUnique(KeyID);
diff --git a/Engine/Source/Editor/UnrealEd/Private/LevelEditorViewport.cpp b/Engine/Source/Editor/UnrealEd/Private/LevelEditorViewport.cpp
--- a/Engine/Source/Editor/UnrealEd/Private/LevelEditorViewport.cpp
+++ b/Engine/Source/Editor/UnrealEd/Private/LevelEditorViewport.cpp
@@ -4570,6 +4570,24 @@ void FLevelEditorViewportClient::ApplyDeltaToSelectedElements(const FTransform&
 	InputState.SetModifierKeyStates(IsShiftPressed(), IsAltPressed(), IsCtrlPressed(), IsCmdPressed());
 
 	FTypedElementListConstRef ElementsToManipulate = GetElementsToManipulate();
+
+	// Making sure manipulated elements have Modify called on them
+	if (GUndo)
+	{
+		ElementsToManipulate->ForEachElement<ITypedElementObjectInterface>([](const TTypedElement<ITypedElementObjectInterface>& ObjectInterface)
+		{
+			if (UObject* const Object = ObjectInterface.GetObject())
+			{
+				if (!GUndo->ContainsObject(Object))
+				{
+					Object->Modify();
+				}
+			}
+
+			return true;
+		});
+	}
+
 	ViewportInteraction->UpdateGizmoManipulation(ElementsToManipulate, GetWidgetMode(), Widget ? Widget->GetCurrentAxis() : EAxisList::None, InputState, ModifiedDeltaTransform);
 }
 
