[ITF Gizmo] Add support for TransformWidgetSizeAdjustment in UI

#Synced-CL 52498372

diff --git a/Engine/Source/Editor/Experimental/EditorInteractiveToolsFramework/Private/EditorGizmos/TransformGizmo.cpp b/Engine/Source/Editor/Experimental/EditorInteractiveToolsFramework/Private/EditorGizmos/TransformGizmo.cpp
--- a/Engine/Source/Editor/Experimental/EditorInteractiveToolsFramework/Private/EditorGizmos/TransformGizmo.cpp
+++ b/Engine/Source/Editor/Experimental/EditorInteractiveToolsFramework/Private/EditorGizmos/TransformGizmo.cpp
@@ -174,15 +174,7 @@ float FTransformGizmoStyle::GetModifiedAxisSizeMultiplier(const TOptional<TUniqu
 		return InModifierOverride.GetValue()(AxisSizeMultiplier);
 	}
 
-	/** Adjustment for TransformGizmoSize to apply to AxisSizeMultiplier. */
-	constexpr float TransformGizmoSizeToMultiplierDivisor = 0.25f;
-
-	const UTransformGizmoEditorSettings* TransformGizmoSettings = GetDefault<UTransformGizmoEditorSettings>();
-
-	constexpr float MinSizeMultiplier = 0.1f;
-
-	// Legacy Gizmo Size is an offset (0-based, default 1), in world units. New is a multiplier, so we need to compensate (1.0 based), such that scaling is in TransformGizmoSizeToMultiplierDivisor increments.
-	return FMath::Max(MinSizeMultiplier, AxisSizeMultiplier * (1.0f + (TransformGizmoSettings->TransformGizmoSize - 1.0f) * TransformGizmoSizeToMultiplierDivisor));
+	return GetDefault<UTransformGizmoEditorSettings>()->GetTransformGizmoAxisMultiplier();
 }
 
 bool FTransformGizmoInteraction::IsCoordinateSystemEnabled(const ECoordSystem InCoordinateSystem) const
diff --git a/Engine/Source/Editor/Experimental/EditorInteractiveToolsFramework/Private/TransformGizmoEditorSettings.cpp b/Engine/Source/Editor/Experimental/EditorInteractiveToolsFramework/Private/TransformGizmoEditorSettings.cpp
--- a/Engine/Source/Editor/Experimental/EditorInteractiveToolsFramework/Private/TransformGizmoEditorSettings.cpp
+++ b/Engine/Source/Editor/Experimental/EditorInteractiveToolsFramework/Private/TransformGizmoEditorSettings.cpp
@@ -235,6 +235,30 @@ void UTransformGizmoEditorSettings::SetTransformGizmoSize(float InTransformGizmo
 	}
 }
 
+void UTransformGizmoEditorSettings::SetTransformGizmoAxisMultiplier(float InAxisMultiplier, bool bInteractive)
+{
+	if (GizmosParameters.Style.AxisSizeMultiplier != InAxisMultiplier)
+	{
+		GizmosParameters.Style.AxisSizeMultiplier = InAxisMultiplier;
+		BroadcastGizmosParametersChange();
+		if (GEditor)
+		{
+			GEditor->RedrawAllViewports();
+		}
+	}
+
+	// Always save when non-interactive
+	if (!bInteractive)
+	{
+		SaveConfig();
+	}
+}
+
+float UTransformGizmoEditorSettings::GetTransformGizmoAxisMultiplier() const
+{
+	return GizmosParameters.Style.AxisSizeMultiplier;
+}
+
 void UTransformGizmoEditorSettings::BroadcastNewTRSGizmoChange() const
 {
 	UEditorInteractiveGizmoManager::OnUsesNewTRSGizmosChangedDelegate().Broadcast(!UsesLegacyGizmo());
diff --git a/Engine/Source/Editor/Experimental/EditorInteractiveToolsFramework/Public/TransformGizmoEditorSettings.h b/Engine/Source/Editor/Experimental/EditorInteractiveToolsFramework/Public/TransformGizmoEditorSettings.h
--- a/Engine/Source/Editor/Experimental/EditorInteractiveToolsFramework/Public/TransformGizmoEditorSettings.h
+++ b/Engine/Source/Editor/Experimental/EditorInteractiveToolsFramework/Public/TransformGizmoEditorSettings.h
@@ -78,6 +78,16 @@ public:
 	UE_API bool UsesNewTRSGizmo() const;
 
 	UE_API void SetGizmosParameters(const FGizmosParameters& InGizmosParameters);
+
+	/** 
+	 * Updates gizmo axis size multiplier. 
+	 * bInteractive set to true means the new value will not be saved to current configuration.
+	 * Useful to get good performance while the value is being interactively changed, e.g. through a UI slider.
+	 */
+	UE_API void SetTransformGizmoAxisMultiplier(float InAxisMultiplier, bool bInteractive = false);
+	UE_API float GetTransformGizmoAxisMultiplier() const;
+
+	/** Use this for legacy gizmos */
 	UE_API void SetTransformGizmoSize(float InTransformGizmoSize);
 
 private:
diff --git a/Engine/Source/Editor/UnrealEd/Private/ViewportToolbar/UnrealEdViewportToolbar.cpp b/Engine/Source/Editor/UnrealEd/Private/ViewportToolbar/UnrealEdViewportToolbar.cpp
--- a/Engine/Source/Editor/UnrealEd/Private/ViewportToolbar/UnrealEdViewportToolbar.cpp
+++ b/Engine/Source/Editor/UnrealEd/Private/ViewportToolbar/UnrealEdViewportToolbar.cpp
@@ -739,6 +739,122 @@ bool IsViewModeSupported(EViewModeIndex InViewModeIndex)
 	}
 }
 
+void AddGizmoSizeSlider(FToolMenuSection& InMenuSection)
+{
+	const FSlateIcon WidgetIcon = FSlateIcon(FAppStyle::GetAppStyleSetName(), "ViewportToolbar.GizmoScale");
+
+	// Helper struct, used by both Legacy and New Gizmo scale sliders to specify their lambdas
+	struct FWidgetLambdas
+	{
+		TFunction<FText()> Tooltip = []() { return FText::GetEmpty(); };
+		TFunction<float()> Value = []() { return 0.0f; };
+		TFunction<void(float)> OnValueChanged = [](float) {};
+		TFunction<void(float, ETextCommit::Type)> OnCommit = TFunction<void(float, ETextCommit::Type)>();
+	};
+
+	// Creates a gizmo scale widget with a slider using the lambdas specified by InWidgetLambdas
+	auto MakeGizmoScaleWidget = [](float InMin, float InMax, FWidgetLambdas InWidgetLambdas)
+	{
+		// Commit lambda is not set, so we use OnValueChanged for that too
+		if (!InWidgetLambdas.OnCommit)
+		{
+			InWidgetLambdas.OnCommit = [OnValueChanged = InWidgetLambdas.OnValueChanged](float InValue, ETextCommit::Type)
+			{
+				OnValueChanged(InValue);
+			};
+		}
+
+		// clang-format off
+		return SNew(SBox)
+		.HAlign(HAlign_Right)
+		[
+			SNew(SBox)
+			.Padding(FMargin(4.0f, 0.0f, 0.0f, 0.0f))
+			.WidthOverride(100.0f)
+			[
+				SNew(SBorder)
+				.BorderImage(FAppStyle::Get().GetBrush("Menu.WidgetBorder"))
+				.Padding(FMargin(1.0f))
+				[
+					SNew(SSpinBox<float>)
+					.Style(&FAppStyle::Get(), "Menu.SpinBox")
+					.Font(FAppStyle::GetFontStyle(TEXT("MenuItem.Font")))
+					.MinValue(InMin)
+					.MaxValue(InMax)
+					.ToolTipText_Lambda(MoveTemp(InWidgetLambdas.Tooltip))
+					.Value_Lambda(MoveTemp(InWidgetLambdas.Value))
+					.OnValueChanged_Lambda(MoveTemp(InWidgetLambdas.OnValueChanged))
+					.OnValueCommitted_Lambda(MoveTemp(InWidgetLambdas.OnCommit))
+				]
+			]
+		];
+		// clang-format on
+	};
+
+	// Legacy Gizmo Scale widget
+	{
+		FWidgetLambdas LegacyWidgetLambdas;
+		LegacyWidgetLambdas.Tooltip = []()
+		{
+			return FText::AsNumber(GetDefault<UTransformGizmoEditorSettings>()->TransformGizmoSize);
+		};
+		LegacyWidgetLambdas.Value = []()
+		{
+			return GetDefault<UTransformGizmoEditorSettings>()->TransformGizmoSize;
+		};
+		LegacyWidgetLambdas.OnValueChanged = [](float InValue)
+		{
+			UTransformGizmoEditorSettings* ViewportSettings = GetMutableDefault<UTransformGizmoEditorSettings>();
+			ViewportSettings->SetTransformGizmoSize(InValue);
+		};
+
+		TSharedRef<SWidget> LegacyGizmoScaleWidget = MakeGizmoScaleWidget(-10.0f,150.0f, LegacyWidgetLambdas);
+
+		FToolMenuEntry& GizmoScaleEntry = InMenuSection.AddEntry(FToolMenuEntry::InitWidget("LegacyGizmoScale", LegacyGizmoScaleWidget, LOCTEXT("LegacyGizmoScaleLabel", "Gizmo Scale")));
+		GizmoScaleEntry.Icon = WidgetIcon;
+		GizmoScaleEntry.Visibility = TAttribute<bool>::CreateLambda([]
+		{
+			// Legacy only
+			return UEditorInteractiveGizmoManager::GetTransformGizmoVersion() == 0;
+		});
+	}
+
+	// New Gizmo Scale widget
+	{
+		FWidgetLambdas WidgetLambdas;
+		WidgetLambdas.Tooltip = []()
+		{
+			return FText::AsNumber(GetDefault<UTransformGizmoEditorSettings>()->GetTransformGizmoAxisMultiplier());
+		};
+		WidgetLambdas.Value = []()
+		{
+			return GetDefault<UTransformGizmoEditorSettings>()->GetTransformGizmoAxisMultiplier();
+		};
+		WidgetLambdas.OnValueChanged = [](float InValue)
+		{
+			// Interactive: SetTransformGizmoAxisMultiplier will not store the new value in config.
+			// Grants better performance while dragging the slider
+			constexpr bool bInteractive = true;
+			UTransformGizmoEditorSettings* ViewportSettings = GetMutableDefault<UTransformGizmoEditorSettings>();
+			ViewportSettings->SetTransformGizmoAxisMultiplier(InValue, bInteractive);
+		};
+		WidgetLambdas.OnCommit = [](float InValue, ETextCommit::Type)
+		{
+			UTransformGizmoEditorSettings* ViewportSettings = GetMutableDefault<UTransformGizmoEditorSettings>();
+			ViewportSettings->SetTransformGizmoAxisMultiplier(InValue);
+		};
+
+		TSharedRef<SWidget> GizmoScaleWidget = MakeGizmoScaleWidget(0.1f,10.0f, WidgetLambdas);
+
+		FToolMenuEntry& GizmoScaleEntry = InMenuSection.AddEntry(FToolMenuEntry::InitWidget("GizmoScale", GizmoScaleWidget, LOCTEXT("GizmoScaleLabel", "Gizmo Scale")));
+		GizmoScaleEntry.Icon = WidgetIcon;
+		GizmoScaleEntry.Visibility = TAttribute<bool>::CreateLambda([]
+		{
+			return UEditorInteractiveGizmoManager::GetTransformGizmoVersion() >= 1;
+		});
+	}
+}
+
 class SOrthographicViewPlaneSlider : public SCompoundWidget
 {
 public:
@@ -1199,66 +1315,7 @@ FToolMenuEntry CreateTransformsSubmenu()
 				{
 					FToolMenuSection& GizmoSection = Submenu->FindOrAddSection("Gizmo", LOCTEXT("GizmoLabel", "Gizmo"));
 
-					{
-						FUIAction Action;
-						Action.ExecuteAction = FExecuteAction::CreateLambda(
-							[]()
-							{
-								UEditorInteractiveGizmoManager::SetUsesNewTRSGizmos(
-									!UEditorInteractiveGizmoManager::UsesNewTRSGizmos()
-								);
-							}
-						);
-						Action.GetActionCheckState = FGetActionCheckState::CreateLambda(
-							[]() -> ECheckBoxState
-							{
-								return UEditorInteractiveGizmoManager::UsesNewTRSGizmos() ? ECheckBoxState::Checked
-																						  : ECheckBoxState::Unchecked;
-							}
-						);
-					}
-
-					
-					TSharedRef<SWidget> GizmoScaleWidget =
-						// clang-format off
-						SNew(SHorizontalBox)
-						+ SHorizontalBox::Slot()
-						.FillWidth(0.9f)
-						[
-							SNew(SSpinBox<int32>)
-								.MinValue(-10.0f)
-								.MaxValue(150.0f)
-								.ToolTipText_Lambda(
-									[]() -> FText
-									{
-										return FText::AsNumber(
-											GetDefault<UTransformGizmoEditorSettings>()->TransformGizmoSize
-										);
-									}
-								)
-								.Value_Lambda(
-									[]() -> float
-									{
-										return GetDefault<UTransformGizmoEditorSettings>()->TransformGizmoSize;
-									}
-								)
-								.OnValueChanged_Lambda(
-									[](float InValue)
-									{
-										UTransformGizmoEditorSettings* ViewportSettings =
-											GetMutableDefault<UTransformGizmoEditorSettings>();
-										ViewportSettings->SetTransformGizmoSize(InValue);
-									}
-								)
-						]
-						+ SHorizontalBox::Slot()
-						.FillWidth(0.1f);
-					// clang-format on
-
-					FToolMenuEntry& GizmoScaleEntry = GizmoSection.AddEntry(FToolMenuEntry::InitWidget(
-						"GizmoScale", GizmoScaleWidget, LOCTEXT("GizmoScaleLabel", "Gizmo Scale")
-					));
-					GizmoScaleEntry.Icon = FSlateIcon(FAppStyle::GetAppStyleSetName(), "ViewportToolbar.GizmoScale");
+					Private::AddGizmoSizeSlider(GizmoSection);
 
 					{
 						FUIAction PreserveNonUniformScaleAction;
