EditorUsability : UnrealEd
- Fixed Editor Preferences - General Appearance application scale not updated after an import

#Synced-CL 52573604

diff --git a/Engine/Source/Editor/UnrealEd/Classes/Settings/EditorStyleSettings.h b/Engine/Source/Editor/UnrealEd/Classes/Settings/EditorStyleSettings.h
--- a/Engine/Source/Editor/UnrealEd/Classes/Settings/EditorStyleSettings.h
+++ b/Engine/Source/Editor/UnrealEd/Classes/Settings/EditorStyleSettings.h
@@ -224,6 +224,7 @@ protected:
 	UNREALED_API virtual void PostEditChangeProperty( struct FPropertyChangedEvent& PropertyChangedEvent ) override;
 
 private:
+	void OnPropertyChanged(FName InPropertyName);
 
 	// Holds an event delegate that is executed when a setting has changed.
 	FSettingChangedEvent SettingChangedEvent;
diff --git a/Engine/Source/Editor/UnrealEd/Private/Settings/EditorStyleClasses.cpp b/Engine/Source/Editor/UnrealEd/Private/Settings/EditorStyleClasses.cpp
--- a/Engine/Source/Editor/UnrealEd/Private/Settings/EditorStyleClasses.cpp
+++ b/Engine/Source/Editor/UnrealEd/Private/Settings/EditorStyleClasses.cpp
@@ -93,26 +93,8 @@ void UEditorStyleSettings::PostEditChangeProperty(struct FPropertyChangedEvent&
 {
 	Super::PostEditChangeProperty(PropertyChangedEvent);
 
-	const FName PropertyName = (PropertyChangedEvent.Property != nullptr) ? PropertyChangedEvent.Property->GetFName() : NAME_None;
-
-
-	// This property is intentionally not per project so it must be manually written to the correct config file
-	if (PropertyName == GET_MEMBER_NAME_CHECKED(UEditorStyleSettings, bEnableHighDPIAwareness))
-	{
-		GConfig->SetBool(TEXT("HDPI"), TEXT("EnableHighDPIAwareness"), bEnableHighDPIAwareness, GEditorSettingsIni);
-	}
-	else if (PropertyName.IsNone() || PropertyName == GET_MEMBER_NAME_CHECKED(UEditorStyleSettings, ApplicationScale))
-	{
-		if (FSlateApplication::IsInitialized())
-		{
-			FSlateApplication::Get().SetApplicationScale(ApplicationScale);
-		}
-	}
-
-//	if (!FUnrealEdMisc::Get().IsDeletePreferences())
-	{
-		SaveConfig();
-	}
+	const FName PropertyName = PropertyChangedEvent.GetPropertyName();
+	OnPropertyChanged(PropertyName);
 
 	SettingChangedEvent.Broadcast(PropertyName);
 }
@@ -128,11 +110,30 @@ bool UEditorStyleSettings::OnImportBegin(const FString& ImportFromPath)
 		FEditorStyleSettingsCustomization::PromptToImportTheme(ImportFromPath); 
 	}
 
-	SaveConfig();
+	OnPropertyChanged(NAME_None);
 
 	return true; 
 }
 
+void UEditorStyleSettings::OnPropertyChanged(FName InPropertyName)
+{
+	// This property is intentionally not per project so it must be manually written to the correct config file
+	if (InPropertyName.IsNone() || InPropertyName == GET_MEMBER_NAME_CHECKED(UEditorStyleSettings, bEnableHighDPIAwareness))
+	{
+		GConfig->SetBool(TEXT("HDPI"), TEXT("EnableHighDPIAwareness"), bEnableHighDPIAwareness, GEditorSettingsIni);
+	}
+	
+	if (InPropertyName.IsNone() || InPropertyName == GET_MEMBER_NAME_CHECKED(UEditorStyleSettings, ApplicationScale))
+	{
+		if (FSlateApplication::IsInitialized())
+		{
+			FSlateApplication::Get().SetApplicationScale(ApplicationScale);
+		}
+	}
+	
+	SaveConfig();
+}
+
 bool UEditorStyleSettings::OnExportBegin(const FString& ExportToPath)
 {
 	SaveConfig(CPF_Config, *ExportToPath);
