[Viewport Toolbar] Low camera speed reset warning button

#Synced-CL 52569872

diff --git a/Engine/Source/Editor/EditorStyle/Private/StarshipStyle.cpp b/Engine/Source/Editor/EditorStyle/Private/StarshipStyle.cpp
--- a/Engine/Source/Editor/EditorStyle/Private/StarshipStyle.cpp
+++ b/Engine/Source/Editor/EditorStyle/Private/StarshipStyle.cpp
@@ -2944,6 +2944,7 @@ void FStarshipEditorStyle::FStyle::SetupNewViewportToolbarStyles()
 	// Viewport Toolbar Icons
 	{
 		Set("ViewportToolbar.CameraSpeed", new IMAGE_BRUSH_SVG("Starship/Common/CameraSpeed_16", Icon16x16));
+		Set("ViewportToolbar.CameraSpeedReset", new IMAGE_BRUSH_SVG("Starship/Common/CameraSpeedReset_16", Icon16x16));
 		Set("ViewportToolbar.EV100", new IMAGE_BRUSH_SVG("Starship/Common/EV100_16", Icon16x16));
 		Set("ViewportToolbar.ExactCameraView", new IMAGE_BRUSH_SVG("Starship/Common/ExactCameraView_16", Icon16x16));
 		Set("ViewportToolbar.Exposure", new IMAGE_BRUSH_SVG("Starship/Common/Exposure_16", Icon16x16));
diff --git a/Engine/Source/Editor/UnrealEd/Private/Settings/EditorViewportSettings.cpp b/Engine/Source/Editor/UnrealEd/Private/Settings/EditorViewportSettings.cpp
--- a/Engine/Source/Editor/UnrealEd/Private/Settings/EditorViewportSettings.cpp
+++ b/Engine/Source/Editor/UnrealEd/Private/Settings/EditorViewportSettings.cpp
@@ -136,3 +136,8 @@ void FEditorViewportCameraSpeedSettings::SetRelativeUISpeed(float RelativeUISpee
 {
 	CurrentSpeed = FMath::Lerp(MinUISpeed, MaxUISpeed, RelativeUISpeed);
 }
+
+float FEditorViewportCameraSpeedSettings::GetLowSpeedWarningThreshold() const
+{
+	return LowSpeedWarningThreshold;
+}
diff --git a/Engine/Source/Editor/UnrealEd/Private/ViewportToolbar/UnrealEdViewportToolbar.cpp b/Engine/Source/Editor/UnrealEd/Private/ViewportToolbar/UnrealEdViewportToolbar.cpp
--- a/Engine/Source/Editor/UnrealEd/Private/ViewportToolbar/UnrealEdViewportToolbar.cpp
+++ b/Engine/Source/Editor/UnrealEd/Private/ViewportToolbar/UnrealEdViewportToolbar.cpp
@@ -4850,6 +4850,62 @@ FToolMenuEntry CreateCameraSpeedMenu()
 					// Don't show camera speed entries (in-menu and raised) if viewport is orthographic (or missing)
 					Entry.SetShowInToolbarTopLevel(true);
 					Entry.Visibility = IsPerspectiveLambda;
+
+					// Reset camera speed entry
+					{
+						// This entry is a toolbar only warning, telling the user that camera speed is extremely low and might cause imperceptible camera movements.
+
+						FUIAction HideInMenuAction;
+						HideInMenuAction.IsActionVisibleDelegate.BindLambda([]{ return false; });
+
+						FToolMenuEntry ResetSpeedEntry = FToolMenuEntry::InitMenuEntry(
+							"LowCameraSpeedWarning",
+							LOCTEXT("LowCameraSpeedWarningLabel", "Reset Camera Speed"),
+							LOCTEXT("LowCameraSpeedWarningTooltip", "Camera speed is extremely low. This might cause imperceptible camera movements.\nClick to reset camera speed to its default."),
+							FSlateIcon(FAppStyle::GetAppStyleSetName(), "ViewportToolbar.CameraSpeedReset"),
+							HideInMenuAction
+						);
+
+						ResetSpeedEntry.ToolBarData.ActionOverride = FToolMenuExecuteAction::CreateLambda(
+							[WeakViewport = Context->Viewport](const FToolMenuContext&)
+							{
+								if (const TSharedPtr<SEditorViewport> ViewportPinned = WeakViewport.Pin())
+								{
+									if (const TSharedPtr<FEditorViewportClient> ViewportClient = ViewportPinned->GetViewportClient())
+									{
+										const float DefaultSpeed = FEditorViewportCameraSpeedSettings{}.GetCurrentSpeed();
+										FEditorViewportCameraSpeedSettings Settings = ViewportClient->GetCameraSpeedSettings();
+										Settings.SetCurrentSpeed(DefaultSpeed);
+										ViewportClient->SetCameraSpeedSettings(Settings);
+									}
+								}
+							}
+						);
+
+						ResetSpeedEntry.SetShowInToolbarTopLevel(
+							TAttribute<bool>::CreateLambda(
+								[WeakViewport = Context->Viewport]() -> bool
+								{
+									if (const TSharedPtr<SEditorViewport> ViewportPinned = WeakViewport.Pin())
+									{
+										if (const TSharedPtr<FEditorViewportClient> ViewportClient = ViewportPinned->GetViewportClient())
+										{
+											if (ViewportClient->IsPerspective())
+											{
+												const float CurrentSpeed = ViewportClient->GetCameraSpeedSettings().GetCurrentSpeed();
+												const float LowSpeedThreshold = ViewportClient->GetCameraSpeedSettings().GetLowSpeedWarningThreshold();
+												return CurrentSpeed <= LowSpeedThreshold;
+											}
+										}
+									}
+									return false;
+								}
+							)
+						);
+						ResetSpeedEntry.StyleNameOverride = "ViewportToolbarWarning";
+
+						InDynamicSection.AddEntry(ResetSpeedEntry);
+					}
 				}
 
 				Entry.ToolBarData.PlacementOverride = MenuPlacement_BelowRightAnchor;
diff --git a/Engine/Source/Editor/UnrealEd/Public/Settings/EditorViewportSettings.h b/Engine/Source/Editor/UnrealEd/Public/Settings/EditorViewportSettings.h
--- a/Engine/Source/Editor/UnrealEd/Public/Settings/EditorViewportSettings.h
+++ b/Engine/Source/Editor/UnrealEd/Public/Settings/EditorViewportSettings.h
@@ -50,6 +50,9 @@ public:
 	/** Set the camera speed with a 0..1 value relative to the min & max */
 	UNREALED_API void SetRelativeUISpeed(float RelativeUISpeed);
 	
+	/** Get the threshold used to warn the user about low camera speed */
+	UNREALED_API float GetLowSpeedWarningThreshold() const;
+	
 private:
 	/** The speed of the camera */
 	UPROPERTY(Meta = (ClampMin = 0.0001f))
@@ -70,4 +73,8 @@ private:
 	/** The maximum speed of the camera when dragging with a slider */
 	UPROPERTY(EditAnywhere, Category = Camera, DisplayName="Max Slider Speed", Meta = (ClampMin = 0.0001f))
 	float MaxUISpeed = 32.0f;
+
+	/** A threshold used to warn the user about low camera speed */
+	UPROPERTY()
+	float LowSpeedWarningThreshold = 0.01f;
 };
