Fix marquee actor selection to not automatically select actors with no components, since they have no spatial bounds and cannot be inside any volume.

#Synced-CL 52452213

diff --git a/Engine/Source/Editor/UnrealEd/Private/Elements/Actor/ActorElementEditorWorldInterface.cpp b/Engine/Source/Editor/UnrealEd/Private/Elements/Actor/ActorElementEditorWorldInterface.cpp
--- a/Engine/Source/Editor/UnrealEd/Private/Elements/Actor/ActorElementEditorWorldInterface.cpp
+++ b/Engine/Source/Editor/UnrealEd/Private/Elements/Actor/ActorElementEditorWorldInterface.cpp
@@ -200,6 +200,11 @@ bool UActorElementEditorWorldInterface::IsElementInConvexVolume(const FTypedElem
 	{
 		TArray<FTypedElementHandle> ChildElementHandles;
 		HierarchyElement.GetChildElements(ChildElementHandles);
+		// An actor with no components has no spatial bounds and cannot be inside any volume.
+		if (ChildElementHandles.IsEmpty())
+		{
+			return false;
+		}
 		for (const FTypedElementHandle& ChildHandle : ChildElementHandles)
 		{
 			if (TTypedElement<ITypedElementWorldInterface> ChildWorldElement = GetRegistry().GetElement<ITypedElementWorldInterface>(ChildHandle))
@@ -233,6 +238,11 @@ bool UActorElementEditorWorldInterface::IsElementInBox(const FTypedElementHandle
 	{
 		TArray<FTypedElementHandle> ChildElementHandles;
 		HierarchyElement.GetChildElements(ChildElementHandles);
+		// An actor with no components has no spatial bounds and cannot be inside any box.
+		if (ChildElementHandles.IsEmpty())
+		{
+			return false;
+		}
 		for (const FTypedElementHandle& ChildHandle : ChildElementHandles)
 		{
 			if (TTypedElement<ITypedElementWorldInterface> ChildWorldElement = GetRegistry().GetElement<ITypedElementWorldInterface>(ChildHandle))
