[Editor] Make the "open source code" action async, so that the editor doesnt lock up every time you press it and wait for your IDE to open. You can disable the "SourceCodeNavigation.EnableAsync" Cvar to go back to a syncronous operation.

#Synced-CL 52582364

diff --git a/Engine/Source/Editor/BlueprintGraph/Private/K2Node_CallFunction.cpp b/Engine/Source/Editor/BlueprintGraph/Private/K2Node_CallFunction.cpp
--- a/Engine/Source/Editor/BlueprintGraph/Private/K2Node_CallFunction.cpp
+++ b/Engine/Source/Editor/BlueprintGraph/Private/K2Node_CallFunction.cpp
@@ -3579,7 +3579,7 @@ void UK2Node_CallFunction::JumpToDefinition() const
 						if (bFileFound)
 						{
 							const FString AbsNativeParentClassHeaderPath = FPaths::ConvertRelativePathToFull(NativeParentClassHeaderPath);
-							bSucceeded = FSourceCodeNavigation::OpenSourceFile(AbsNativeParentClassHeaderPath);
+							FSourceCodeNavigation::OpenSourceFileAsync(AbsNativeParentClassHeaderPath);
 						}
 					}
 				}
diff --git a/Engine/Source/Editor/GameProjectGeneration/Private/SNewClassDialog.cpp b/Engine/Source/Editor/GameProjectGeneration/Private/SNewClassDialog.cpp
--- a/Engine/Source/Editor/GameProjectGeneration/Private/SNewClassDialog.cpp
+++ b/Engine/Source/Editor/GameProjectGeneration/Private/SNewClassDialog.cpp
@@ -920,7 +920,7 @@ void SNewClassDialog::OnEditCodeClicked()
 	if (ClassHeaderPath.Len() > 0)
 	{
 		const FString AbsoluteHeaderPath = IFileManager::Get().ConvertToAbsolutePathForExternalAppForRead(*ClassHeaderPath);
-		FSourceCodeNavigation::OpenSourceFile(AbsoluteHeaderPath);
+		FSourceCodeNavigation::OpenSourceFileAsync(AbsoluteHeaderPath);
 	}
 }
 
@@ -1245,7 +1245,7 @@ void SNewClassDialog::FinishClicked()
 					SourceFiles.Add(IFileManager::Get().ConvertToAbsolutePathForExternalAppForRead(*HeaderFilePath));
 					SourceFiles.Add(IFileManager::Get().ConvertToAbsolutePathForExternalAppForRead(*CppFilePath));
 
-					FSourceCodeNavigation::OpenSourceFiles(SourceFiles);
+					FSourceCodeNavigation::OpenSourceFilesAsync(SourceFiles);
 				}
 			}
 
diff --git a/Engine/Source/Editor/LevelEditor/Private/LevelEditorToolBar.cpp b/Engine/Source/Editor/LevelEditor/Private/LevelEditorToolBar.cpp
--- a/Engine/Source/Editor/LevelEditor/Private/LevelEditorToolBar.cpp
+++ b/Engine/Source/Editor/LevelEditor/Private/LevelEditorToolBar.cpp
@@ -282,7 +282,7 @@ namespace LevelEditorActionHelpers
 			if (bFileFound)
 			{
 				const FString AbsoluteHeaderPath = IFileManager::Get().ConvertToAbsolutePathForExternalAppForRead(*NativeParentClassHeaderPath);
-				FSourceCodeNavigation::OpenSourceFile( AbsoluteHeaderPath );
+				FSourceCodeNavigation::OpenSourceFileAsync( AbsoluteHeaderPath );
 			}
 		}
 	}
diff --git a/Engine/Source/Editor/LevelEditor/Private/SActorDetails.cpp b/Engine/Source/Editor/LevelEditor/Private/SActorDetails.cpp
--- a/Engine/Source/Editor/LevelEditor/Private/SActorDetails.cpp
+++ b/Engine/Source/Editor/LevelEditor/Private/SActorDetails.cpp
@@ -1164,7 +1164,7 @@ void SActorDetails::OnNativeComponentWarningHyperlinkClicked(const FSlateHyperli
 		if (bFileFound)
 		{
 			const FString AbsoluteHeaderPath = IFileManager::Get().ConvertToAbsolutePathForExternalAppForRead(*NativeParentClassHeaderPath);
-			FSourceCodeNavigation::OpenSourceFile(AbsoluteHeaderPath);
+			FSourceCodeNavigation::OpenSourceFileAsync(AbsoluteHeaderPath);
 		}
 	}
 }
diff --git a/Engine/Source/Editor/UnrealEd/Private/Factories/PackFactory.cpp b/Engine/Source/Editor/UnrealEd/Private/Factories/PackFactory.cpp
--- a/Engine/Source/Editor/UnrealEd/Private/Factories/PackFactory.cpp
+++ b/Engine/Source/Editor/UnrealEd/Private/Factories/PackFactory.cpp
@@ -836,7 +836,7 @@ UObject* UPackFactory::FactoryCreateBinary
 					const FText Message = NSLOCTEXT("PackFactory", "CodeAdded", "Added source file(s). Would you like to edit the code now?");
 					if ( FMessageDialog::Open(EAppMsgType::YesNo, Message) == EAppReturnType::Yes )
 					{
-						FSourceCodeNavigation::OpenSourceFiles(WrittenSourceFiles);
+						FSourceCodeNavigation::OpenSourceFilesAsync(WrittenSourceFiles);
 					}
 				}
 			}
diff --git a/Engine/Source/Editor/UnrealEd/Private/Kismet2/CompilerResultsLog.cpp b/Engine/Source/Editor/UnrealEd/Private/Kismet2/CompilerResultsLog.cpp
--- a/Engine/Source/Editor/UnrealEd/Private/Kismet2/CompilerResultsLog.cpp
+++ b/Engine/Source/Editor/UnrealEd/Private/Kismet2/CompilerResultsLog.cpp
@@ -601,7 +601,7 @@ void FCompilerResultsLog::OnGotoError(const TSharedRef<IMessageToken>& Token)
 		LineNumberString.LeftChopInline(1, EAllowShrinking::No); // remove right parenthesis
 		int32 LineNumber = FCString::Strtoi(*LineNumberString, NULL, 10);
 
-		FSourceCodeNavigation::OpenSourceFile( FullPath, LineNumber );
+		FSourceCodeNavigation::OpenSourceFileAsync( FullPath, LineNumber );
 	}
 }
 
diff --git a/Engine/Source/Editor/UnrealEd/Private/Kismet2/ComponentEditorUtils.cpp b/Engine/Source/Editor/UnrealEd/Private/Kismet2/ComponentEditorUtils.cpp
--- a/Engine/Source/Editor/UnrealEd/Private/Kismet2/ComponentEditorUtils.cpp
+++ b/Engine/Source/Editor/UnrealEd/Private/Kismet2/ComponentEditorUtils.cpp
@@ -1252,7 +1252,7 @@ void FComponentEditorUtils::OnGoToComponentAssetInBrowser(UObject* Asset)
 void FComponentEditorUtils::OnOpenComponentCodeFile(const FString CodeFileName)
 {
 	const FString AbsoluteHeaderPath = IFileManager::Get().ConvertToAbsolutePathForExternalAppForRead(*CodeFileName);
-	FSourceCodeNavigation::OpenSourceFile(AbsoluteHeaderPath);
+	FSourceCodeNavigation::OpenSourceFileAsync(AbsoluteHeaderPath);
 }
 
 void FComponentEditorUtils::OnEditBlueprintComponent(UObject* Blueprint)
diff --git a/Engine/Source/Editor/UnrealEd/Private/SourceCodeNavigation.cpp b/Engine/Source/Editor/UnrealEd/Private/SourceCodeNavigation.cpp
--- a/Engine/Source/Editor/UnrealEd/Private/SourceCodeNavigation.cpp
+++ b/Engine/Source/Editor/UnrealEd/Private/SourceCodeNavigation.cpp
@@ -2,6 +2,7 @@
 
 #include "SourceCodeNavigation.h"
 
+#include "Async/Async.h"
 #include "Async/AsyncWork.h"
 #include "Async/TaskGraphInterfaces.h"
 #include "Editor/UnrealEdEngine.h"
@@ -57,6 +58,14 @@
 
 DEFINE_LOG_CATEGORY(LogSelectionDetails);
 
+static bool GAsyncSourceCodeNavigation = true;
+static FAutoConsoleVariableRef CVarAsyncSourceCodeNavigation(
+	TEXT("SourceCodeNavigation.EnableAsync"),
+	GAsyncSourceCodeNavigation,
+	TEXT("When true, opening source files in the IDE is deferred so the editor stays responsive. Set to false to revert to synchronous (blocking) behavior."),
+	ECVF_Default
+);
+
 #define LOCTEXT_NAMESPACE "SourceCodeNavigation"
 
 #define SOURCECODENAVIGATOR_SHOW_CONSTRUCTOR_AND_DESTRUCTOR 0 // @todo editcode: Not sure if we want this by default (make user-configurable?)
@@ -1513,7 +1522,7 @@ bool FSourceCodeNavigation::NavigateToClass(const UClass* InClass)
 	if (FSourceCodeNavigation::FindClassHeaderPath(InClass, ClassHeaderPath) && IFileManager::Get().FileSize(*ClassHeaderPath) != INDEX_NONE)
 	{
 		FString AbsoluteHeaderPath = IFileManager::Get().ConvertToAbsolutePathForExternalAppForRead(*ClassHeaderPath);
-		FSourceCodeNavigation::OpenSourceFile(AbsoluteHeaderPath);
+		FSourceCodeNavigation::OpenSourceFileAsync(AbsoluteHeaderPath);
 		return true;
 	}
 	return false;
@@ -1563,7 +1572,7 @@ bool FSourceCodeNavigation::NavigateToStruct(const UScriptStruct* InStruct)
 	if (FSourceCodeNavigation::FindClassHeaderPath(InStruct, StructHeaderPath) && IFileManager::Get().FileSize(*StructHeaderPath) != INDEX_NONE)
 	{
 		FString AbsoluteHeaderPath = IFileManager::Get().ConvertToAbsolutePathForExternalAppForRead(*StructHeaderPath);
-		FSourceCodeNavigation::OpenSourceFile(AbsoluteHeaderPath);
+		FSourceCodeNavigation::OpenSourceFileAsync(AbsoluteHeaderPath);
 		return true;
 	}
 	return false;
@@ -1677,7 +1686,8 @@ bool FSourceCodeNavigation::NavigateToProperty(const FProperty* InProperty)
 		if (bFileLocated)
 		{
 			const FString AbsoluteSourcePath = IFileManager::Get().ConvertToAbsolutePathForExternalAppForRead(*SourceFilePath);
-			return OpenSourceFile( AbsoluteSourcePath );
+			OpenSourceFileAsync(AbsoluteSourcePath);
+			return true;
 		}
 	}
 	return false;
@@ -1727,7 +1737,7 @@ bool FSourceCodeNavigation::NavigateToStruct(const UStruct* InStruct)
 	if (FSourceCodeNavigation::FindClassHeaderPath(InStruct, ClassHeaderPath) && IFileManager::Get().FileSize(*ClassHeaderPath) != INDEX_NONE)
 	{
 		FString AbsoluteHeaderPath = IFileManager::Get().ConvertToAbsolutePathForExternalAppForRead(*ClassHeaderPath);
-		FSourceCodeNavigation::OpenSourceFile(AbsoluteHeaderPath);
+		FSourceCodeNavigation::OpenSourceFileAsync(AbsoluteHeaderPath);
 		return true;
 	}
 
@@ -1909,6 +1919,106 @@ bool FSourceCodeNavigation::OpenSourceFiles(const TArray<FString>& AbsoluteSourc
 	return false;
 }
 
+void FSourceCodeNavigation::OpenSourceFileAsync(const FString& AbsoluteSourcePath, int32 LineNumber, int32 ColumnNumber)
+{
+	if (!GAsyncSourceCodeNavigation)
+	{
+		OpenSourceFile(AbsoluteSourcePath, LineNumber, ColumnNumber);
+		return;
+	}
+
+	if (!IsCompilerAvailable())
+	{
+		AccessOnCompilerNotFound().Broadcast();
+		return;
+	}
+
+	FNotificationInfo Info(LOCTEXT("OpeningSourceFile", "Opening source file..."));
+	Info.Image = FAppStyle::GetBrush(TEXT("LevelEditor.RecompileGameCode"));
+	Info.bFireAndForget = false;
+	Info.ExpireDuration = 2.0f;
+	Info.bUseThrobber = true;
+
+	TSharedPtr<SNotificationItem> NotificationPtr = FSlateNotificationManager::Get().AddNotification(Info);
+	if (NotificationPtr.IsValid())
+	{
+		NotificationPtr->SetCompletionState(SNotificationItem::CS_Pending);
+	}
+
+	// Defer to the next game thread tick so the calling UI event handler returns immediately
+	// and Slate can render a frame before the (potentially blocking) accessor call executes.
+	// NOTE: We must stay on the game thread because some accessors (e.g. Rider) may show
+	// Slate dialogs internally, which requires game thread access.
+	FString PathCopy = AbsoluteSourcePath;
+	AsyncTask(ENamedThreads::GameThread, [PathCopy, LineNumber, ColumnNumber, NotificationPtr]()
+	{
+		ISourceCodeAccessModule& SourceCodeAccessModule = FModuleManager::LoadModuleChecked<ISourceCodeAccessModule>("SourceCodeAccess");
+		const bool bSuccess = SourceCodeAccessModule.GetAccessor().OpenFileAtLine(PathCopy, LineNumber, ColumnNumber);
+
+		if (NotificationPtr.IsValid())
+		{
+			if (bSuccess)
+			{
+				NotificationPtr->SetCompletionState(SNotificationItem::CS_Success);
+			}
+			else
+			{
+				NotificationPtr->SetText(LOCTEXT("OpenSourceFileFailed", "Failed to open source file"));
+				NotificationPtr->SetCompletionState(SNotificationItem::CS_Fail);
+			}
+			NotificationPtr->ExpireAndFadeout();
+		}
+	});
+}
+
+void FSourceCodeNavigation::OpenSourceFilesAsync(const TArray<FString>& AbsoluteSourcePaths)
+{
+	if (!GAsyncSourceCodeNavigation)
+	{
+		OpenSourceFiles(AbsoluteSourcePaths);
+		return;
+	}
+
+	if (!IsCompilerAvailable())
+	{
+		AccessOnCompilerNotFound().Broadcast();
+		return;
+	}
+
+	FNotificationInfo Info(LOCTEXT("OpeningSourceFiles", "Opening source files..."));
+	Info.Image = FAppStyle::GetBrush(TEXT("LevelEditor.RecompileGameCode"));
+	Info.bFireAndForget = false;
+	Info.ExpireDuration = 2.0f;
+	Info.bUseThrobber = true;
+
+	TSharedPtr<SNotificationItem> NotificationPtr = FSlateNotificationManager::Get().AddNotification(Info);
+	if (NotificationPtr.IsValid())
+	{
+		NotificationPtr->SetCompletionState(SNotificationItem::CS_Pending);
+	}
+
+	TArray<FString> PathsCopy = AbsoluteSourcePaths;
+	AsyncTask(ENamedThreads::GameThread, [PathsCopy = MoveTemp(PathsCopy), NotificationPtr]()
+	{
+		ISourceCodeAccessModule& SourceCodeAccessModule = FModuleManager::LoadModuleChecked<ISourceCodeAccessModule>("SourceCodeAccess");
+		const bool bSuccess = SourceCodeAccessModule.GetAccessor().OpenSourceFiles(PathsCopy);
+
+		if (NotificationPtr.IsValid())
+		{
+			if (bSuccess)
+			{
+				NotificationPtr->SetCompletionState(SNotificationItem::CS_Success);
+			}
+			else
+			{
+				NotificationPtr->SetText(LOCTEXT("OpenSourceFilesFailed", "Failed to open source files"));
+				NotificationPtr->SetCompletionState(SNotificationItem::CS_Fail);
+			}
+			NotificationPtr->ExpireAndFadeout();
+		}
+	});
+}
+
 bool FSourceCodeNavigation::AddSourceFiles(const TArray<FString>& AbsoluteSourcePaths)
 {
 	if ( IsCompilerAvailable() )
diff --git a/Engine/Source/Editor/UnrealEd/Public/SourceCodeNavigation.h b/Engine/Source/Editor/UnrealEd/Public/SourceCodeNavigation.h
--- a/Engine/Source/Editor/UnrealEd/Public/SourceCodeNavigation.h
+++ b/Engine/Source/Editor/UnrealEd/Public/SourceCodeNavigation.h
@@ -357,6 +357,12 @@ public:
 	/** Opens a multiple source files */
 	UNREALED_API static bool OpenSourceFiles(const TArray<FString>& AbsoluteSourcePaths);
 
+	/** Opens a single source file asynchronously, showing a notification while the IDE loads. */
+	UNREALED_API static void OpenSourceFileAsync(const FString& AbsoluteSourcePath, int32 LineNumber = 0, int32 ColumnNumber = 0);
+
+	/** Opens multiple source files asynchronously, showing a notification while the IDE loads. */
+	UNREALED_API static void OpenSourceFilesAsync(const TArray<FString>& AbsoluteSourcePaths);
+
 	/** Add multiple source files to the current solution/project/workspace */
 	UNREALED_API static bool AddSourceFiles(const TArray<FString>& AbsoluteSourcePaths);
 
