Cooker: Replace function static TSet<FString> that stores whether a package has been loaded for cooking with a bool on FPackageData, for usability and performance.
#rnx
#Synced-CL 52568437

diff --git a/Engine/Source/Editor/UnrealEd/Private/CookOnTheFlyServer.cpp b/Engine/Source/Editor/UnrealEd/Private/CookOnTheFlyServer.cpp
--- a/Engine/Source/Editor/UnrealEd/Private/CookOnTheFlyServer.cpp
+++ b/Engine/Source/Editor/UnrealEd/Private/CookOnTheFlyServer.cpp
@@ -4663,11 +4663,10 @@ bool UCookOnTheFlyServer::LoadPackageForCooking(UE::Cook::FPackageData& PackageD
 #if DEBUG_COOKONTHEFLY
 	UE_LOG(LogCook, Display, TEXT("Processing request %s"), *ReportingFileName);
 #endif
-	static TSet<FString> CookWarningsList;
 	TOptional<TInheritedGuardValueAccessors<GetOutputCookingWarnings, SetOutputCookingWarnings>> OutputCookingWarningsGuard;
-	if (CookWarningsList.Contains(FileName) == false)
+	if (!PackageData.HasBeenLoadedForCooking())
 	{
-		CookWarningsList.Add(FileName);
+		PackageData.SetHasBeenLoadedForCooking(true);
 		OutputCookingWarningsGuard.Emplace(IsCookFlagSet(ECookInitializationFlags::OutputVerboseCookerWarnings));
 	}
 
diff --git a/Engine/Source/Editor/UnrealEd/Private/Cooker/CookPackageData.cpp b/Engine/Source/Editor/UnrealEd/Private/Cooker/CookPackageData.cpp
--- a/Engine/Source/Editor/UnrealEd/Private/Cooker/CookPackageData.cpp
+++ b/Engine/Source/Editor/UnrealEd/Private/Cooker/CookPackageData.cpp
@@ -116,7 +116,7 @@ FPackageData::FPackageData(FPackageDatas& PackageDatas, const FName& InPackageNa
 	, bGenerated(0), bKeepReferencedDuringGC(0)
 	, bWasCookedThisSession(0)
 	, DoesGeneratedRequireGeneratorValue(static_cast<uint32>(ICookPackageSplitter::EGeneratedRequiresGenerator::None))
-	, bHasReplayedLogMessages(0)
+	, bHasReplayedLogMessages(0), bHasBeenLoadedForCooking(0)
 {
 	SetState(EPackageState::Idle);
 	SetSaveSubState(ESaveSubState::StartSave);
@@ -489,6 +489,7 @@ void FPackageData::ClearCookResults()
 	}
 	SetSuppressCookReason(ESuppressCookReason::NotSuppressed);
 	bHasReplayedLogMessages = false;
+	bHasBeenLoadedForCooking = false;
 }
 
 void FPackageData::ClearCookResults(const ITargetPlatform* TargetPlatform)
@@ -514,6 +515,7 @@ void FPackageData::ClearCookResults(const ITargetPlatform* TargetPlatform)
 		bWasCookedThisSession = false;
 		PackageDatas.GetMonitor().OnLastCookedPlatformRemoved(*this);
 		bHasReplayedLogMessages = false;
+		bHasBeenLoadedForCooking = false;
 	}
 }
 
diff --git a/Engine/Source/Editor/UnrealEd/Private/Cooker/CookPackageData.h b/Engine/Source/Editor/UnrealEd/Private/Cooker/CookPackageData.h
--- a/Engine/Source/Editor/UnrealEd/Private/Cooker/CookPackageData.h
+++ b/Engine/Source/Editor/UnrealEd/Private/Cooker/CookPackageData.h
@@ -828,6 +828,8 @@ public:
 	void ClearLogMessages();
 	bool HasReplayedLogMessages() const;
 	void SetHasReplayedLogMessages(bool bValue);
+	bool HasBeenLoadedForCooking() const;
+	void SetHasBeenLoadedForCooking(bool Value);
 
 	/**
 	 * Return the platforms for which the given Package has been marked reachable.
@@ -943,6 +945,7 @@ private:
 	static_assert(static_cast<uint32>(ICookPackageSplitter::EGeneratedRequiresGenerator::Count) <= 4, "We are storing Enum value in 2 bits");
 	uint32 DoesGeneratedRequireGeneratorValue : 2;
 	uint32 bHasReplayedLogMessages : 1;
+	uint32 bHasBeenLoadedForCooking : 1;
 };
 
 /**
@@ -1989,6 +1992,16 @@ inline void FPackageData::SetHasReplayedLogMessages(bool Value)
 	bHasReplayedLogMessages = Value != 0;
 }
 
+inline bool FPackageData::HasBeenLoadedForCooking() const
+{
+	return static_cast<bool>(bHasBeenLoadedForCooking);
+}
+
+inline void FPackageData::SetHasBeenLoadedForCooking(bool Value)
+{
+	bHasBeenLoadedForCooking = Value != 0;
+}
+
 inline FWorkerId FPackageData::GetWorkerAssignment() const
 {
 	return WorkerAssignment;
