Editor: Don't assume that viewport detachment means that the PIE window was closed

Allows PIE Game Instances to exist sans Viewport without being interpreted as an attempt to close PIE.  The editor already reliably detects when PIE windows have actually been closed, even if through OS-level messaging.

#rnx
#Synced-CL 52600201

diff --git a/Engine/Source/Editor/UnrealEd/Private/EditorEngine.cpp b/Engine/Source/Editor/UnrealEd/Private/EditorEngine.cpp
--- a/Engine/Source/Editor/UnrealEd/Private/EditorEngine.cpp
+++ b/Engine/Source/Editor/UnrealEd/Private/EditorEngine.cpp
@@ -1766,21 +1766,6 @@ void UEditorEngine::Tick( float DeltaSeconds, bool bIdleMode )
 	// Clean up the game viewports that have been closed.
 	CleanupGameViewport();
 
-	// If all viewports closed, close the current play level.
-	if( PlayWorld && !bIsSimulatingInEditor )
-	{
-		for (auto It=WorldList.CreateIterator(); It; ++It)
-		{
-			// For now, kill PIE session if any of the viewports are closed
-			if (It->WorldType == EWorldType::PIE && It->GameViewport == NULL && !It->RunAsDedicated && !It->bWaitingOnOnlineSubsystem)
-			{
-				EndPlayMap();
-				break;
-			}
-		}
-	}
-
-
 	// Potentially rebuilds the streaming data.
 	EditorContext.World()->ConditionallyBuildStreamingData();
 
@@ -2243,16 +2228,6 @@ void UEditorEngine::Tick( float DeltaSeconds, bool bIdleMode )
 	// Clean up any game viewports that may have been closed during the level tick (eg by Kismet).
 	CleanupGameViewport();
 
-	// If all viewports closed, close the current play level.
-	if( GameViewport == NULL && PlayWorld && !bIsSimulatingInEditor )
-	{
-		FWorldContext& PieWorldContext = GetWorldContextFromWorldChecked(PlayWorld);
-		if (!PieWorldContext.RunAsDedicated && !PieWorldContext.bWaitingOnOnlineSubsystem)
-		{
-			EndPlayMap();
-		}
-	}
-
 	// Updates all the extensions for all the editor worlds
 	EditorWorldExtensionsManager->Tick( DeltaSeconds );
 
