Cooker: CookPackageSplitter: Set the GCookerIsLoadingPackage variable to true when calling CookPackageSplitter functions, since those functions can load packages.
#rnx
#Synced-CL 52610419

diff --git a/Engine/Source/Editor/UnrealEd/Private/CookOnTheFlyServer.cpp b/Engine/Source/Editor/UnrealEd/Private/CookOnTheFlyServer.cpp
--- a/Engine/Source/Editor/UnrealEd/Private/CookOnTheFlyServer.cpp
+++ b/Engine/Source/Editor/UnrealEd/Private/CookOnTheFlyServer.cpp
@@ -3933,6 +3933,8 @@ UE::Cook::EPollStatus UCookOnTheFlyServer::PreSaveGeneratorPackage(UE::Cook::FPa
 			FName()
 #endif
 		);
+		TGuardValue<bool> ScopedIsCookerLoadingPackage(GIsCookerLoadingPackage, true);
+
 		ICookPackageSplitter::FPopulateContext PopulateContext(PopulateData);
 		PopulateData.OwnerPackage = Package;
 		PopulateData.OwnerObject = SplitDataObject;
@@ -4102,6 +4104,8 @@ UE::Cook::EPollStatus UCookOnTheFlyServer::TryPreSaveGeneratedPackage(UE::Cook::
 			FName()
 #endif
 		);
+		TGuardValue<bool> ScopedIsCookerLoadingPackage(GIsCookerLoadingPackage, true);
+
 		ICookPackageSplitter::FPopulateContext PopulateContext(PopulateData);
 		PopulateData.OwnerPackage = OwnerPackage;
 		PopulateData.OwnerObject = OwnerObject;
diff --git a/Engine/Source/Editor/UnrealEd/Private/Cooker/CookGenerationHelper.cpp b/Engine/Source/Editor/UnrealEd/Private/Cooker/CookGenerationHelper.cpp
--- a/Engine/Source/Editor/UnrealEd/Private/Cooker/CookGenerationHelper.cpp
+++ b/Engine/Source/Editor/UnrealEd/Private/Cooker/CookGenerationHelper.cpp
@@ -992,6 +992,8 @@ bool FGenerationHelper::TryReportGenerationManifest()
 			FName()
 #endif
 		);
+		TGuardValue<bool> ScopedIsCookerLoadingPackage(GIsCookerLoadingPackage, true);
+
 		Manifest = GetCookPackageSplitterInstance()->ReportGenerationManifest(LocalOwnerPackage, OwnerObject);
 		if (Manifest.GeneratedPackages.IsEmpty())
 		{
@@ -1273,6 +1275,7 @@ bool FGenerationHelper::TryCallPopulateGeneratorPackage(
 			FName()
 #endif
 		);
+		TGuardValue<bool> ScopedIsCookerLoadingPackage(GIsCookerLoadingPackage, true);
 
 		// Get the list of Objects in the package. Read them from the PackageData if it already has them,
 		// otherwise calculate them into a temporary buffer.
@@ -1462,6 +1465,8 @@ bool FGenerationHelper::TryCallPopulateGeneratedPackage(UE::Cook::FCookGeneratio
 		FName()
 #endif
 	);
+	TGuardValue<bool> ScopedIsCookerLoadingPackage(GIsCookerLoadingPackage, true);
+
 	UE::Cook::CookPackageSplitter::FPopulateContextData PopulateData;
 	ICookPackageSplitter::FPopulateContext PopulateContext(PopulateData);
 	PopulateData.OwnerPackage = GetOwnerPackage();
@@ -2081,6 +2086,7 @@ void FGenerationHelper::ResetSaveState(FCookGenerationInfo& Info, UPackage* Pack
 				FName()
 #endif
 			);
+			TGuardValue<bool> ScopedIsCookerLoadingPackage(GIsCookerLoadingPackage, true);
 
 			UE::Cook::CookPackageSplitter::FPopulateContextData PopulateData;
 			ICookPackageSplitter::FPopulateContext PopulateContext(PopulateData);
