Teds Actor Folders: Fix a bug where renaming a folder would cause the folder label to get erased in TEDS

#Synced-CL 52629034

diff --git a/Engine/Plugins/Experimental/EditorDataStorageFeatures/Source/TedsOutliner/Private/Widgets/OutlinerFolderLabelWidget.cpp b/Engine/Plugins/Experimental/EditorDataStorageFeatures/Source/TedsOutliner/Private/Widgets/OutlinerFolderLabelWidget.cpp
--- a/Engine/Plugins/Experimental/EditorDataStorageFeatures/Source/TedsOutliner/Private/Widgets/OutlinerFolderLabelWidget.cpp
+++ b/Engine/Plugins/Experimental/EditorDataStorageFeatures/Source/TedsOutliner/Private/Widgets/OutlinerFolderLabelWidget.cpp
@@ -165,7 +165,8 @@ TSharedRef<SWidget> FOutlinerFolderLabelWidgetConstructor::CreateLabel(UE::Edito
 			.OnExitEditingMode_Lambda([DataStorage, WidgetRow]()
 			{
 				DataStorage->RemoveColumn<FIsInEditingModeTag>(WidgetRow);
-			});
+			})
+			.DoubleSelectDelay(0.f);
 
 		DataStorage->AddColumn<FWidgetEnterEditModeColumn>(WidgetRow, FWidgetEnterEditModeColumn{
 				.OnEnterEditMode = FSimpleDelegate::CreateLambda([TextBlock]()
diff --git a/Engine/Source/Editor/UnrealEd/Private/ActorFolders/ActorFolderSyncQueries.cpp b/Engine/Source/Editor/UnrealEd/Private/ActorFolders/ActorFolderSyncQueries.cpp
--- a/Engine/Source/Editor/UnrealEd/Private/ActorFolders/ActorFolderSyncQueries.cpp
+++ b/Engine/Source/Editor/UnrealEd/Private/ActorFolders/ActorFolderSyncQueries.cpp
@@ -48,7 +48,7 @@ void UTedsActorFolderSyncFactory::RegisterLabelQueries(UE::Editor::DataStorage::
 	InDataStorage.RegisterQuery(
 		Select(
 			TEXT("Sync folder name to label column"),
-			FProcessor(EQueryTickPhase::PostPhysics, InDataStorage.GetQueryTickGroupName(EQueryTickGroups::SyncExternalToDataStorage))
+			FProcessor(EQueryTickPhase::PrePhysics, InDataStorage.GetQueryTickGroupName(EQueryTickGroups::SyncDataStorageToExternal))
 				.SetExecutionMode(EExecutionMode::GameThread),
 			[this](IQueryContext& Context, RowHandle RowHandle, const FTypedElementLabelColumn& LabelColumn,
 				const FTypedElementLabelHashColumn& LabelHashColumn, const FFolderCompatibilityColumn& FolderColumn,
@@ -71,7 +71,7 @@ void UTedsActorFolderSyncFactory::RegisterLabelQueries(UE::Editor::DataStorage::
 	InDataStorage.RegisterQuery(
 		Select(
 			TEXT("Sync label column to folder name"),
-			FProcessor(EQueryTickPhase::PostPhysics, InDataStorage.GetQueryTickGroupName(EQueryTickGroups::SyncDataStorageToExternal))
+			FProcessor(EQueryTickPhase::PostPhysics, InDataStorage.GetQueryTickGroupName(EQueryTickGroups::SyncExternalToDataStorage))
 				.SetExecutionMode(EExecutionMode::GameThread),
 			[](IQueryContext& Context, RowHandle Row, const FFolderCompatibilityColumn& FolderColumn,
 				FTypedElementLabelColumn& LabelColumn, FTypedElementLabelHashColumn& LabelHashColumn)
