# NewMaterialTranslator

Fix 3 separate bugs with the UI presentation of pins and connections in the material graph.

## Bugs
1. Inactive material output pins weren't showing up gray/disabled.
2. RGBA output pins for mask outputs weren't color coded red, green, or blue.
3. Inline editing widgets on the material output pin weren't showing up.

## Cause
When the NMT was enabled UMaterialGraph::UpdatePinTypes would be called which would stomp the PinCategory and PinSubCategory to add type information.  The existing values for PinCategory are used to indentify MaterialInputs and for the  pin widget factory to route pins to the correct Slate widget type.  The PinSubCategory is used to handle mask colors, and type to determine inline widgets.  These can't be repurposed.

## Fix
Instead of stomping PinCategory and PinSubCategory we use the PinSubCategoryObject, which is a completely new object type that can be used to carry any information we want.  This now caches the value off from the material insights and can be used be the pin/connection color code to determine colors instead.  In the future we can attach more metadata to this object.

## Changes
- Remove unused Void and Value types.  These are legacy and were added by the "old" new material translator.
- Adds  UMaterialShaderValueTypeObject type which is used to carry type and metadata information for connections.  This is filled out on the PinSubCategoryObject in UpdatePinTypes.  The slate color functions use this object to determine colors if it exists.
- Added GetPinColor() override to UMaterialGraphSchema that reads UMaterialShaderValueTypeObject from PinSubCategoryObject for insight-based type coloring, falling back to GetPinTypeColor() for legacy pins.
- Update FMaterialGraphConnectionDrawingPolicy::DetermineWiringStyle, SGraphPinMaterialInput::GetPinColor, UMaterialGraphSchema::GetPinColor,  to use the PinSubCategoryObject to get coloring.
- Set UsePinColorForText(false) unconditionally on SGraphPinMaterialInput to prevent insight-derived pin colors from bleeding into pin label text.

#rnx
#Synced-CL 52629756

diff --git a/Engine/Source/Editor/MaterialEditor/Private/MaterialGraphConnectionDrawingPolicy.cpp b/Engine/Source/Editor/MaterialEditor/Private/MaterialGraphConnectionDrawingPolicy.cpp
--- a/Engine/Source/Editor/MaterialEditor/Private/MaterialGraphConnectionDrawingPolicy.cpp
+++ b/Engine/Source/Editor/MaterialEditor/Private/MaterialGraphConnectionDrawingPolicy.cpp
@@ -10,6 +10,7 @@
 #include "Layout/ArrangedWidget.h"
 #include "MaterialGraph/MaterialGraph.h"
 #include "MaterialGraph/MaterialGraphSchema.h"
+#include "MaterialGraph/MaterialShaderValueTypeObject.h"
 #include "MaterialGraphNode_Knot.h"
 #include "Math/Color.h"
 #include "SGraphPin.h"
@@ -131,11 +132,8 @@ void FMaterialGraphConnectionDrawingPolicy::DetermineWiringStyle(UEdGraphPin* Ou
 		Params.WireColor = FSubstrateWidget::GetConnectionColor();
 	}
 
-	UE::Shader::EValueType InputType = UE::Shader::EValueType::Void;
-	UE::Shader::EValueType OutputType = UE::Shader::EValueType::Void;
 	bool bInactivePin = false;
 	bool bExecPin = false;
-	bool bValueTypePin = false;
 
 	// Have to consider both pins as the input will be an 'output' when previewing a connection
 	if (OutputPin)
@@ -149,15 +147,6 @@ void FMaterialGraphConnectionDrawingPolicy::DetermineWiringStyle(UEdGraphPin* Ou
 		{
 			bExecPin = true;
 		}
-		else if (OutputPin->PinType.PinCategory == UMaterialGraphSchema::PC_Void)
-		{
-			bValueTypePin = true;
-		}
-		else if (OutputPin->PinType.PinCategory == UMaterialGraphSchema::PC_ValueType)
-		{
-			OutputType = UE::Shader::FindValueType(OutputPin->PinType.PinSubCategory);
-			bValueTypePin = true;
-		}
 
 		UEdGraphNode* OutputNode = OutputPin->GetOwningNode();
 		if (UMaterialGraphNode_Knot* OutputKnotNode = Cast<UMaterialGraphNode_Knot>(OutputNode))
@@ -183,15 +172,6 @@ void FMaterialGraphConnectionDrawingPolicy::DetermineWiringStyle(UEdGraphPin* Ou
 		{
 			bExecPin = true;
 		}
-		else if (InputPin->PinType.PinCategory == UMaterialGraphSchema::PC_Void)
-		{
-			bValueTypePin = true;
-		}
-		else if (InputPin->PinType.PinCategory == UMaterialGraphSchema::PC_ValueType)
-		{
-			InputType = UE::Shader::FindValueType(InputPin->PinType.PinSubCategory);
-			bValueTypePin = true;
-		}
 
 		UEdGraphNode* InputNode = InputPin->GetOwningNode();
 		if (UMaterialGraphNode_Knot* InputKnotNode = Cast<UMaterialGraphNode_Knot>(InputNode))
@@ -216,10 +196,13 @@ void FMaterialGraphConnectionDrawingPolicy::DetermineWiringStyle(UEdGraphPin* Ou
 		Params.WireColor = Settings->ExecutionPinTypeColor;
 		Params.WireThickness = Settings->DefaultExecutionWireThickness;
 	}
-	else if (bValueTypePin)
+	else if (InputPin)
 	{
-		// TODO - how to handle InputPin being nullptr, when does this occur?
-		if (InputPin)
+		if (const UMaterialShaderValueTypeObject* ValueTypeObj = Cast<UMaterialShaderValueTypeObject>(InputPin->PinType.PinSubCategoryObject.Get()))
+		{
+			Params.WireColor = UMaterialGraphSchema::GetColorForConnectionType(ValueTypeObj->ValueType);
+		}
+		else
 		{
 			Params.WireColor = MaterialGraphSchema->GetPinTypeColor(InputPin->PinType);
 		}
diff --git a/Engine/Source/Editor/MaterialEditor/Private/MaterialPins/SGraphPinMaterialInput.cpp b/Engine/Source/Editor/MaterialEditor/Private/MaterialPins/SGraphPinMaterialInput.cpp
--- a/Engine/Source/Editor/MaterialEditor/Private/MaterialPins/SGraphPinMaterialInput.cpp
+++ b/Engine/Source/Editor/MaterialEditor/Private/MaterialPins/SGraphPinMaterialInput.cpp
@@ -7,6 +7,7 @@
 #include "EdGraph/EdGraphPin.h"
 #include "MaterialGraph/MaterialGraph.h"
 #include "MaterialGraph/MaterialGraphSchema.h"
+#include "MaterialGraph/MaterialShaderValueTypeObject.h"
 #include "Misc/AssertionMacros.h"
 #include "Templates/Casts.h"
 #include "MaterialEditor/SGraphSubstrateMaterial.h"
@@ -38,7 +39,11 @@ FSlateColor SGraphPinMaterialInput::GetPinColor() const
 			if (UMaterialGraphSchema::GetMaterialIOValueType(GraphPinObj) == MCT_Substrate)
 			{
 				return FSlateColor(FSubstrateWidget::GetConnectionColor());
-			}			
+			}
+		}
+		if (const UMaterialShaderValueTypeObject* ValueTypeObj = Cast<UMaterialShaderValueTypeObject>(GraphPinObj->PinType.PinSubCategoryObject.Get()))
+		{
+			return FSlateColor(UMaterialGraphSchema::GetColorForConnectionType(ValueTypeObj->ValueType));
 		}
 		return Schema->ActivePinColor;
 	}
diff --git a/Engine/Source/Editor/UnrealEd/Classes/MaterialGraph/MaterialGraph.h b/Engine/Source/Editor/UnrealEd/Classes/MaterialGraph/MaterialGraph.h
--- a/Engine/Source/Editor/UnrealEd/Classes/MaterialGraph/MaterialGraph.h
+++ b/Engine/Source/Editor/UnrealEd/Classes/MaterialGraph/MaterialGraph.h
@@ -13,11 +13,6 @@
 class UMaterialExpressionComment;
 class UMaterialExpressionComposite;
 
-namespace UE::Shader
-{
-enum class EValueType : uint8;
-}
-
 DECLARE_DELEGATE_RetVal( bool, FRealtimeStateGetter );
 DECLARE_DELEGATE( FSetMaterialDirty );
 DECLARE_DELEGATE_OneParam( FToggleExpressionCollapsed, UMaterialExpression* );
diff --git a/Engine/Source/Editor/UnrealEd/Classes/MaterialGraph/MaterialGraphSchema.h b/Engine/Source/Editor/UnrealEd/Classes/MaterialGraph/MaterialGraphSchema.h
--- a/Engine/Source/Editor/UnrealEd/Classes/MaterialGraph/MaterialGraphSchema.h
+++ b/Engine/Source/Editor/UnrealEd/Classes/MaterialGraph/MaterialGraphSchema.h
@@ -13,6 +13,11 @@ class UEdGraph;
 class UMaterialExpression;
 struct FToolMenuContext;
 
+namespace UE::Shader
+{
+	enum class EValueType : uint8;
+}
+
 /** Action to add an expression node to the graph */
 USTRUCT()
 struct FMaterialGraphSchemaAction_NewNode : public FEdGraphSchemaAction
@@ -197,8 +202,6 @@ class UMaterialGraphSchema : public UEdGraphSchema
 	UNREALED_API static const FName PC_Optional;
 	UNREALED_API static const FName PC_MaterialInput;
 	UNREALED_API static const FName PC_Exec;
-	UNREALED_API static const FName PC_Void;
-	UNREALED_API static const FName PC_ValueType;
 
 	// Common PinType.PinSubCategory values
 	UNREALED_API static const FName PSC_Red;
@@ -270,6 +273,10 @@ class UMaterialGraphSchema : public UEdGraphSchema
 	virtual const FPinConnectionResponse CanCreateConnection(const UEdGraphPin* A, const UEdGraphPin* B) const override;
 	virtual bool TryCreateConnection(UEdGraphPin* A, UEdGraphPin* B) const override;
 	virtual FLinearColor GetPinTypeColor(const FEdGraphPinType& PinType) const override;
+	virtual FLinearColor GetPinColor(const UEdGraphPin* InPin) const override;
+
+	/** Get the color for a connection based on its resolved shader value type (from material insights) */
+	UNREALED_API static FLinearColor GetColorForConnectionType(UE::Shader::EValueType ConnectionType);
 	virtual void BreakNodeLinks(UEdGraphNode& TargetNode) const override;
 	virtual void BreakPinLinks(UEdGraphPin& TargetPin, bool bSendsNodeNotifcation) const override;
 	virtual void BreakSinglePinLink(UEdGraphPin* SourcePin, UEdGraphPin* TargetPin) const override;
diff --git a/Engine/Source/Editor/UnrealEd/Classes/MaterialGraph/MaterialShaderValueTypeObject.h b/Engine/Source/Editor/UnrealEd/Classes/MaterialGraph/MaterialShaderValueTypeObject.h
new file mode 100644
--- /dev/null
+++ b/Engine/Source/Editor/UnrealEd/Classes/MaterialGraph/MaterialShaderValueTypeObject.h
@@ -0,0 +1,28 @@
+// Copyright Epic Games, Inc. All Rights Reserved.
+
+#pragma once
+
+#include "CoreMinimal.h"
+#include "UObject/ObjectMacros.h"
+#include "Shader/ShaderTypes.h"
+#include "MaterialShaderValueTypeObject.generated.h"
+
+/**
+ * Lightweight flyweight UObject that wraps a UE::Shader::EValueType.
+ * One singleton instance per EValueType value, stored in PinType.PinSubCategoryObject
+ * to carry shader value type information without modifying PinCategory or PinSubCategory.
+ * Used by the new material translator to drive pin and wire coloring.
+ *
+ * Owned by UMaterialShaderValueTypeSubsystem, whose lifetime is tied to GEditor.
+ */
+UCLASS(MinimalAPI, Transient)
+class UMaterialShaderValueTypeObject : public UObject
+{
+	GENERATED_BODY()
+
+public:
+	UE::Shader::EValueType ValueType = UE::Shader::EValueType::Void;
+
+	/** Get or create the flyweight singleton for a given EValueType. Returns nullptr for EValueType::Void or if GEditor is not available. */
+	UNREALED_API static UMaterialShaderValueTypeObject* Get(UE::Shader::EValueType InType);
+};
diff --git a/Engine/Source/Editor/UnrealEd/Private/MaterialGraph.cpp b/Engine/Source/Editor/UnrealEd/Private/MaterialGraph.cpp
--- a/Engine/Source/Editor/UnrealEd/Private/MaterialGraph.cpp
+++ b/Engine/Source/Editor/UnrealEd/Private/MaterialGraph.cpp
@@ -27,6 +27,7 @@
 #include "Materials/MaterialFunction.h"
 #include "Materials/MaterialInsights.h"
 
+#include "MaterialGraph/MaterialShaderValueTypeObject.h"
 #include "MaterialGraphNode_Knot.h"
 
 #include "Kismet2/BlueprintEditorUtils.h"
@@ -123,47 +124,47 @@ UMaterialGraphNode_Base* FindGraphNodeForObject(const UObject* Object)
 
 void UMaterialGraph::UpdatePinTypes()
 {
-	if (Material->IsUsingNewHLSLGenerator())
+	if (!Material || !Material->IsUsingNewHLSLGenerator())
 	{
-		const FMaterialInsights* Insight = Material->MaterialInsight.Get();
-		if (!Insight)
-		{
-			return;
-		}
+		return;
+	}
+
+	const FMaterialInsights* Insight = Material->MaterialInsight.Get();
+	if (!Insight)
+	{
+		return;
+	}
 
-		// Reset all pins to void category.
-		for (UEdGraphNode* Node :Nodes)
+	// Clear any previously-set value type objects, preserving UEnum objects used for inline editing.
+	for (UEdGraphNode* Node : Nodes)
+	{
+		for (UEdGraphPin* Pin : Node->Pins)
 		{
-			for (UEdGraphPin* Pin : Node->Pins)
+			if (Cast<UMaterialShaderValueTypeObject>(Pin->PinType.PinSubCategoryObject.Get()))
 			{
-				if (Pin->PinType.PinCategory != UMaterialGraphSchema::PC_Exec)
-				{
-					Pin->PinType.PinCategory = UMaterialGraphSchema::PC_Void;
-				}
+				Pin->PinType.PinSubCategoryObject = nullptr;
 			}
 		}
-		
-		for (const FMaterialInsights::FConnectionInsight& ConnectionInsight : Insight->ConnectionInsights)
+	}
+
+	for (const FMaterialInsights::FConnectionInsight& ConnectionInsight : Insight->ConnectionInsights)
+	{
+		UMaterialShaderValueTypeObject* ValueTypeObj = UMaterialShaderValueTypeObject::Get(ConnectionInsight.ValueType);
+
+		if (UMaterialGraphNode_Base* InputNode = ::Private::FindGraphNodeForObject(ConnectionInsight.InputObject))
 		{
-			if (UMaterialGraphNode_Base* InputNode = ::Private::FindGraphNodeForObject(ConnectionInsight.InputObject))
+			const int32 InputIndex = InputNode->GetSourceIndexForInputIndex(ConnectionInsight.InputIndex);
+			if (UEdGraphPin* InputPin = InputNode->GetInputPin(InputIndex))
 			{
-				const int32 InputIndex = InputNode->GetSourceIndexForInputIndex(ConnectionInsight.InputIndex);
-				if (UEdGraphPin* InputPin = InputNode->GetInputPin(InputIndex))
-				{
-					const UE::Shader::FValueTypeDescription InputTypeDesc = UE::Shader::GetValueTypeDescription(ConnectionInsight.ValueType);
-					InputPin->PinType.PinCategory = UMaterialGraphSchema::PC_ValueType;
-					InputPin->PinType.PinSubCategory = InputTypeDesc.Name;
-				}
+				InputPin->PinType.PinSubCategoryObject = ValueTypeObj;
 			}
+		}
 
-			if (UMaterialGraphNode_Base* OutputNode = ::Private::FindGraphNodeForObject(ConnectionInsight.OutputExpression))
+		if (UMaterialGraphNode_Base* OutputNode = ::Private::FindGraphNodeForObject(ConnectionInsight.OutputExpression))
+		{
+			if (UEdGraphPin* OutputPin = OutputNode->GetOutputPin(ConnectionInsight.OutputIndex))
 			{
-				if (UEdGraphPin* OutputPin = OutputNode->GetOutputPin(ConnectionInsight.OutputIndex))
-				{
-					const UE::Shader::FValueTypeDescription InputTypeDesc = UE::Shader::GetValueTypeDescription(ConnectionInsight.ValueType);
-					OutputPin->PinType.PinCategory = UMaterialGraphSchema::PC_ValueType;
-					OutputPin->PinType.PinSubCategory = InputTypeDesc.Name;
-				}
+				OutputPin->PinType.PinSubCategoryObject = ValueTypeObj;
 			}
 		}
 	}
diff --git a/Engine/Source/Editor/UnrealEd/Private/MaterialGraphSchema.cpp b/Engine/Source/Editor/UnrealEd/Private/MaterialGraphSchema.cpp
--- a/Engine/Source/Editor/UnrealEd/Private/MaterialGraphSchema.cpp
+++ b/Engine/Source/Editor/UnrealEd/Private/MaterialGraphSchema.cpp
@@ -48,6 +48,7 @@
 #include "Editor/AssetReferenceFilter.h"
 #include "IAssetTools.h"
 #include "MaterialEditorActions.h"
+#include "MaterialGraph/MaterialShaderValueTypeObject.h"
 #include "MaterialGraphNode_Knot.h"
 #include "RenderUtils.h"
 #include "MaterialEditorSettings.h"
@@ -272,8 +273,6 @@ const FName UMaterialGraphSchema::PC_Required(TEXT("required"));
 const FName UMaterialGraphSchema::PC_Optional(TEXT("optional"));
 const FName UMaterialGraphSchema::PC_MaterialInput(TEXT("materialinput"));
 const FName UMaterialGraphSchema::PC_Exec(TEXT("exec"));
-const FName UMaterialGraphSchema::PC_Void(TEXT("void"));
-const FName UMaterialGraphSchema::PC_ValueType(TEXT("value"));
 
 const FName UMaterialGraphSchema::PSC_Red(TEXT("red"));
 const FName UMaterialGraphSchema::PSC_Green(TEXT("green"));
@@ -616,11 +615,10 @@ bool UMaterialGraphSchema::TryCreateConnection(UEdGraphPin* A, UEdGraphPin* B) c
 	return bModified;
 }
 
-namespace Private
-{
-FLinearColor GetColorForConnectionType(const UGraphEditorSettings* Settings, UE::Shader::EValueType ConnectionType)
+FLinearColor UMaterialGraphSchema::GetColorForConnectionType(UE::Shader::EValueType ConnectionType)
 {
 	using namespace UE::Shader;
+	const UGraphEditorSettings* Settings = GetDefault<UGraphEditorSettings>();
 	if (ConnectionType == EValueType::Any)
 	{
 		return Settings->WildcardPinTypeColor;
@@ -673,7 +671,6 @@ FLinearColor GetColorForConnectionType(const UGraphEditorSettings* Settings, UE:
 
 	return Settings->DefaultPinTypeColor;
 }
-} // namespace Private
 
 FLinearColor UMaterialGraphSchema::GetPinTypeColor(const FEdGraphPinType& PinType) const
 {
@@ -704,20 +701,23 @@ FLinearColor UMaterialGraphSchema::GetPinTypeColor(const FEdGraphPinType& PinTyp
 	{
 		return InactivePinColor;
 	}
-	else if (PinType.PinCategory == PC_ValueType)
-	{
-		const UE::Shader::EValueType ValueType = UE::Shader::FindValueType(PinType.PinSubCategory);
-		const UGraphEditorSettings* Settings = GetDefault<UGraphEditorSettings>();
-		return Private::GetColorForConnectionType(Settings, ValueType);
-	}
-	else if (PinType.PinCategory == PC_Void)
-	{
-		return InactivePinColor;
-	}
 
 	return ActivePinColor;
 }
 
+FLinearColor UMaterialGraphSchema::GetPinColor(const UEdGraphPin* InPin) const
+{
+	if (InPin)
+	{
+		if (const UMaterialShaderValueTypeObject* ValueTypeObj = Cast<UMaterialShaderValueTypeObject>(InPin->PinType.PinSubCategoryObject.Get()))
+		{
+			return GetColorForConnectionType(ValueTypeObj->ValueType);
+		}
+		return GetPinTypeColor(InPin->PinType);
+	}
+	return FLinearColor::Black;
+}
+
 void UMaterialGraphSchema::BreakNodeLinks(UEdGraphNode& TargetNode) const
 {
 	bool bHasLinksToBreak = false;
diff --git a/Engine/Source/Editor/UnrealEd/Private/MaterialShaderValueTypeObject.cpp b/Engine/Source/Editor/UnrealEd/Private/MaterialShaderValueTypeObject.cpp
new file mode 100644
--- /dev/null
+++ b/Engine/Source/Editor/UnrealEd/Private/MaterialShaderValueTypeObject.cpp
@@ -0,0 +1,28 @@
+// Copyright Epic Games, Inc. All Rights Reserved.
+
+#include "MaterialGraph/MaterialShaderValueTypeObject.h"
+#include "Subsystems/MaterialShaderValueTypeSubsystem.h"
+#include "Editor.h"
+
+UMaterialShaderValueTypeObject* UMaterialShaderValueTypeObject::Get(UE::Shader::EValueType InType)
+{
+	check(IsInGameThread());
+
+	if (InType == UE::Shader::EValueType::Void)
+	{
+		return nullptr;
+	}
+
+	if (!GEditor)
+	{
+		return nullptr;
+	}
+
+	UMaterialShaderValueTypeSubsystem* Subsystem = GEditor->GetEditorSubsystem<UMaterialShaderValueTypeSubsystem>();
+	if (!Subsystem)
+	{
+		return nullptr;
+	}
+
+	return Subsystem->GetOrCreate(InType);
+}
diff --git a/Engine/Source/Editor/UnrealEd/Private/MaterialShaderValueTypeSubsystem.cpp b/Engine/Source/Editor/UnrealEd/Private/MaterialShaderValueTypeSubsystem.cpp
new file mode 100644
--- /dev/null
+++ b/Engine/Source/Editor/UnrealEd/Private/MaterialShaderValueTypeSubsystem.cpp
@@ -0,0 +1,15 @@
+// Copyright Epic Games, Inc. All Rights Reserved.
+
+#include "Subsystems/MaterialShaderValueTypeSubsystem.h"
+#include "MaterialGraph/MaterialShaderValueTypeObject.h"
+
+UMaterialShaderValueTypeObject* UMaterialShaderValueTypeSubsystem::GetOrCreate(UE::Shader::EValueType InType)
+{
+	TObjectPtr<UMaterialShaderValueTypeObject>& Slot = Flyweights.FindOrAdd(static_cast<uint8>(InType));
+	if (!Slot)
+	{
+		Slot = NewObject<UMaterialShaderValueTypeObject>(this);
+		Slot->ValueType = InType;
+	}
+	return Slot;
+}
diff --git a/Engine/Source/Editor/UnrealEd/Public/Subsystems/MaterialShaderValueTypeSubsystem.h b/Engine/Source/Editor/UnrealEd/Public/Subsystems/MaterialShaderValueTypeSubsystem.h
new file mode 100644
--- /dev/null
+++ b/Engine/Source/Editor/UnrealEd/Public/Subsystems/MaterialShaderValueTypeSubsystem.h
@@ -0,0 +1,28 @@
+// Copyright Epic Games, Inc. All Rights Reserved.
+
+#pragma once
+
+#include "EditorSubsystem.h"
+#include "Shader/ShaderTypes.h"
+#include "MaterialShaderValueTypeSubsystem.generated.h"
+
+class UMaterialShaderValueTypeObject;
+
+/**
+ * Editor subsystem that owns the flyweight UMaterialShaderValueTypeObject instances.
+ * Lifetime tied to GEditor — flyweights are GC-traced via UPROPERTY and collected
+ * naturally when the subsystem is destroyed during GEditor teardown.
+ */
+UCLASS(MinimalAPI)
+class UMaterialShaderValueTypeSubsystem : public UEditorSubsystem
+{
+	GENERATED_BODY()
+
+public:
+	/** Get or create the flyweight for a given EValueType. */
+	UMaterialShaderValueTypeObject* GetOrCreate(UE::Shader::EValueType InType);
+
+private:
+	UPROPERTY()
+	TMap<uint8, TObjectPtr<UMaterialShaderValueTypeObject>> Flyweights;
+};
