Editor: Promote Slate.bAllowThrottling to Editor Preferences

* Expose Slate.bAllowThrottling as a checkbox in Editor Preferences > Performance > Editor Performance so artists can discover and toggle it without using the console.

#Synced-CL 52423541

diff --git a/Engine/Source/Editor/UnrealEd/Classes/Editor/EditorPerformanceSettings.h b/Engine/Source/Editor/UnrealEd/Classes/Editor/EditorPerformanceSettings.h
--- a/Engine/Source/Editor/UnrealEd/Classes/Editor/EditorPerformanceSettings.h
+++ b/Engine/Source/Editor/UnrealEd/Classes/Editor/EditorPerformanceSettings.h
@@ -73,6 +73,10 @@ class UEditorPerformanceSettings : public UDeveloperSettings
 	UPROPERTY(EditAnywhere, config, Category=EditorPerformance, meta=(DisplayName="Use Less CPU when in Background") )
 	uint32 bThrottleCPUWhenNotForeground:1;
 
+	/** When enabled, the editor throttles background processing (world ticking, viewport updates) during UI interactions to keep menus and panels responsive. Disable this to keep viewports and simulations running at full speed while interacting with the editor. */
+	UPROPERTY(EditAnywhere, config, Category=EditorPerformance, meta=(DisplayName="Allow Editor (Slate) Throttling", ConsoleVariable="Slate.bAllowThrottling"))
+	uint32 bAllowSlateThrottling:1;
+
 	/** Disables realtime viewports by default when connected via a remote session */
 	UPROPERTY(EditAnywhere, config, Category = EditorPerformance, meta = (DisplayName = "Disable realtime viewports by default in Remote Sessions"))
 	uint32 bDisableRealtimeViewportsInRemoteSessions : 1;
diff --git a/Engine/Source/Editor/UnrealEd/Private/EditorPerformanceSettings.cpp b/Engine/Source/Editor/UnrealEd/Private/EditorPerformanceSettings.cpp
--- a/Engine/Source/Editor/UnrealEd/Private/EditorPerformanceSettings.cpp
+++ b/Engine/Source/Editor/UnrealEd/Private/EditorPerformanceSettings.cpp
@@ -19,6 +19,7 @@ UEditorPerformanceSettings::UEditorPerformanceSettings(const FObjectInitializer&
 	: Super(ObjectInitializer)
 	, bShowFrameRateAndMemory(false)
 	, bThrottleCPUWhenNotForeground(true)
+	, bAllowSlateThrottling(true)
 	, bDisableRealtimeViewportsInRemoteSessions(true)
 	, bMonitorEditorPerformance(false)
 	, bEnableScalabilityWarningIndicator(true)
@@ -51,6 +52,11 @@ void UEditorPerformanceSettings::PostInitProperties()
 		ExportValuesToConsoleVariables(EnableVSyncProperty);
 	}
 
+	if (FProperty* AllowSlateThrottlingProperty = GetClass()->FindPropertyByName(GET_MEMBER_NAME_CHECKED(UEditorPerformanceSettings, bAllowSlateThrottling)))
+	{
+		ExportValuesToConsoleVariables(AllowSlateThrottlingProperty);
+	}
+
 	UEditorPerformanceProjectSettings::ExportResolutionValuesToConsoleVariables();
 }
 
