Cooker: Fix "Cooker has been blocked from saving" message to include diagnostics when the blockage is caused by GeneratorPackages waiting on a response from the CookDirector rather than waiting on a UObject to return true from IsCachedCookedPlatformDataLoaded.
#rnx
#Synced-CL 52610419

diff --git a/Engine/Source/Editor/UnrealEd/Private/CookOnTheFlyServer.cpp b/Engine/Source/Editor/UnrealEd/Private/CookOnTheFlyServer.cpp
--- a/Engine/Source/Editor/UnrealEd/Private/CookOnTheFlyServer.cpp
+++ b/Engine/Source/Editor/UnrealEd/Private/CookOnTheFlyServer.cpp
@@ -1767,56 +1767,83 @@ void UCookOnTheFlyServer::SetSaveBusy(bool bInBusy, int32 NumAsyncWorkRetired)
 		FMsg::Logf(__FILE__, __LINE__, LogCook.GetCategoryName(), MessageSeverity, TEXT("%s"), *Message);
 #endif
 
-		UE_LOGF(LogCook, Display, "%d packages in the savequeue: ", SaveQueue.Num());
 		int DisplayCount = 0;
 		const int DisplayMax = 10;
-		for (TSet<UPackage*>* PackageSet : { &NonExpectedPackages, &ExpectedPackages })
+		if (NonExpectedPackages.Num() + ExpectedPackages.Num() > 0)
 		{
-			for (UPackage* Package : *PackageSet)
+			UE_LOGF(LogCook, Display, "%d packages in the savequeue: ", SaveQueue.Num());
+			for (TSet<UPackage*>* PackageSet : { &NonExpectedPackages, &ExpectedPackages })
 			{
-				if (DisplayCount == DisplayMax)
+				for (UPackage* Package : *PackageSet)
 				{
-					UE_LOGF(LogCook, Display, "    ...");
-					break;
+					if (DisplayCount == DisplayMax)
+					{
+						UE_LOGF(LogCook, Display, "    ...");
+						break;
+					}
+					UE_LOGF(LogCook, Display, "    %ls", *Package->GetName());
+					++DisplayCount;
 				}
-				UE_LOGF(LogCook, Display, "    %ls", *Package->GetName());
-				++DisplayCount;
 			}
-		}
-		if (DisplayCount == 0)
-		{
-			UE_LOGF(LogCook, Display, "    <None>");
-		}
+			if (DisplayCount == 0)
+			{
+				UE_LOGF(LogCook, Display, "    <None>");
+			}
+
+			UE_LOGF(LogCook, Display, "%d objects that have not yet returned true from IsCachedCookedPlatformDataLoaded:",
+				PackageDatas->GetPendingCookedPlatformDataNum());
+			DisplayCount = 0;
+			for (TArray<UObject*>* ObjectArray : { &NonExpectedObjects, &ExpectedObjects })
+			{
+				for (UObject* Object : *ObjectArray)
+				{
+					if (DisplayCount == DisplayMax)
+					{
+						UE_LOGF(LogCook, Display, "    ...");
+						break;
+					}
+					UE_LOGF(LogCook, Display, "    %ls", *Object->GetFullName());
 
-		UE_LOGF(LogCook, Display, "%d objects that have not yet returned true from IsCachedCookedPlatformDataLoaded:",
-			PackageDatas->GetPendingCookedPlatformDataNum());
-		DisplayCount = 0;
-		for (TArray<UObject*>* ObjectArray : {  &NonExpectedObjects, &ExpectedObjects })
+					TStringBuilder<2048> AdditionalDebugInfo;
+					FDelegates::PackageBlocked.Broadcast(Object, AdditionalDebugInfo);
+					UE::String::ParseTokens(AdditionalDebugInfo.ToView(), TEXT('\n'),
+						[](FStringView Line)
+						{
+							UE_LOGF(LogCook, Display, "        %.*ls", Line.Len(), Line.GetData());
+						}, UE::String::EParseTokensOptions::SkipEmpty);
+
+					++DisplayCount;
+				}
+			}
+			if (DisplayCount == 0)
+			{
+				UE_LOGF(LogCook, Display, "    <None>");
+			}
+		}
+		else if (SaveQueue.Num() > 0)
 		{
-			for (UObject* Object : *ObjectArray)
+			UE_LOGF(LogCook, Display, "%d packages in the savequeue: ", SaveQueue.Num());
+			for (FPackageData* PackageData : SaveQueue)
 			{
 				if (DisplayCount == DisplayMax)
 				{
 					UE_LOGF(LogCook, Display, "    ...");
 					break;
 				}
-				UE_LOGF(LogCook, Display, "    %ls", *Object->GetFullName());
-
-				TStringBuilder<2048> AdditionalDebugInfo;
-				FDelegates::PackageBlocked.Broadcast(Object, AdditionalDebugInfo);
-				UE::String::ParseTokens(AdditionalDebugInfo.ToView(), TEXT('\n'),
-					[](FStringView Line)
-					{
-						UE_LOGF(LogCook, Display, "        %.*ls", Line.Len(), Line.GetData());
-					}, UE::String::EParseTokensOptions::SkipEmpty);
-
+				FString GeneratorText;
+				if (PackageData->GetGenerationHelper())
+				{
+					GeneratorText = TEXT(" (Generator)");
+				}
+				else if (PackageData->GetParentGenerationHelper())
+				{
+					GeneratorText = TEXT(" (Generated)");
+				}
+				UE_LOGF(LogCook, Display, "    %ls%ls: %ls",
+					*PackageData->GetPackageName().ToString(), *GeneratorText, LexToString(PackageData->GetSaveSubState()));
 				++DisplayCount;
 			}
 		}
-		if (DisplayCount == 0)
-		{
-			UE_LOGF(LogCook, Display, "    <None>");
-		}
 
 		SaveBusyWarnTimeSeconds = CurrentTime + GCookProgressWarnBusyTime;
 			
