[ITF Viewport Interactions] Fix issues where flying the camera did not send the correct property changed events for piloted actors.

#rnx
#Synced-CL 52694953

diff --git a/Engine/Source/Editor/Experimental/EditorInteractiveToolsFramework/Public/ViewportInteractions/ViewportInteractionsBehaviorSource.h b/Engine/Source/Editor/Experimental/EditorInteractiveToolsFramework/Public/ViewportInteractions/ViewportInteractionsBehaviorSource.h
--- a/Engine/Source/Editor/Experimental/EditorInteractiveToolsFramework/Public/ViewportInteractions/ViewportInteractionsBehaviorSource.h
+++ b/Engine/Source/Editor/Experimental/EditorInteractiveToolsFramework/Public/ViewportInteractions/ViewportInteractionsBehaviorSource.h
@@ -259,7 +259,7 @@ public:
 		return bIsRightMouseButtonDown;
 	}
 
-	bool IsMouseLooking() const;
+	UE_API bool IsMouseLooking() const;
 
 	UE_API void SetIsMouseLooking(bool bInIsLooking);
 
diff --git a/Engine/Source/Editor/UnrealEd/Private/EditorViewportClient.cpp b/Engine/Source/Editor/UnrealEd/Private/EditorViewportClient.cpp
--- a/Engine/Source/Editor/UnrealEd/Private/EditorViewportClient.cpp
+++ b/Engine/Source/Editor/UnrealEd/Private/EditorViewportClient.cpp
@@ -2344,7 +2344,7 @@ void FEditorViewportClient::UpdateCameraMovement( float DeltaTime )
 		{
 			bHasMovement = true;
 		}
-		bHasMovement = bHasMovement && bIsTracking;
+		bHasMovement = bHasMovement && (bIsITFInteractions || bIsTracking);
 
 		BeginCameraMovement(bHasMovement);
 
@@ -2422,10 +2422,10 @@ bool FEditorViewportClient::ConsumeExternalCameraMovement()
 	FRotator Rotation = FRotator::ZeroRotator;
 	FVector Orbit = FVector::ZeroVector;
 	
-	BeginCameraMovement(true);
-	
 	ViewportNavigationHelper->ConsumeTransformDelta(Translation, Rotation, Orbit);
 	
+	BeginCameraMovement(!Translation.IsNearlyZero() || !Rotation.IsNearlyZero() || !Orbit.IsNearlyZero());
+	
 	if (IsPerspective())
 	{
 		bool bCameraHasMoved = false;
diff --git a/Engine/Source/Editor/UnrealEd/Private/LevelEditorViewport.cpp b/Engine/Source/Editor/UnrealEd/Private/LevelEditorViewport.cpp
--- a/Engine/Source/Editor/UnrealEd/Private/LevelEditorViewport.cpp
+++ b/Engine/Source/Editor/UnrealEd/Private/LevelEditorViewport.cpp
@@ -2654,9 +2654,32 @@ void FLevelEditorViewportClient::BeginCameraMovement(bool bHasMovement)
 			bIsCameraMoving = true;
 		}
 	}
-	else if (bIsUsingLegacyMovementNotify || !bIsTracking)
+	else
 	{
-		bIsCameraMoving = false;
+		const bool bDisableCameraMovement = [&]
+		{
+			if (bIsUsingLegacyMovementNotify)
+			{
+				return true;
+			}
+			if (ModeTools->UsesViewportITF())
+			{
+				if (UModeManagerInteractiveToolsContext* Context = ModeTools->GetInteractiveToolsContext())
+				{
+					if (const UViewportInteractionsBehaviorSource* Source = Context->GetViewportInteractionsBehaviorSource())
+					{
+						return !Source->IsMouseLooking();
+					}
+				}
+				return true;
+			}
+			return !bIsTracking;
+		}();
+	
+		if (bDisableCameraMovement)
+		{
+			bIsCameraMoving = false;
+		}
 	}
 }
 
