[Viewport Interactions] Add a CVar to control whether the "Circle Cross" fallback cursor is presented when ITF has an active mouse behavior and no cursor override. Disabled by default.

#rnx
#Synced-CL 52734355

diff --git a/Engine/Source/Editor/UnrealEd/Private/Tools/EdModeInteractiveToolsContext.cpp b/Engine/Source/Editor/UnrealEd/Private/Tools/EdModeInteractiveToolsContext.cpp
--- a/Engine/Source/Editor/UnrealEd/Private/Tools/EdModeInteractiveToolsContext.cpp
+++ b/Engine/Source/Editor/UnrealEd/Private/Tools/EdModeInteractiveToolsContext.cpp
@@ -68,12 +68,20 @@ namespace UE::Editor::Private
 		TEXT("Whether to force on viewport interactions, ignoring legacy requests from editor modes.")
 	);
 	
+	TAutoConsoleVariable<int32> CVarCursorOverrideFallback(
+		TEXT("Editor.ITF.CursorOverrideFallback"),
+		0,
+		TEXT("How to handle the cursor when an ITF behavior has mouse capture but there is no cursor override.\n")
+		TEXT("0 - Fall back to the legacy cursor determination path.\n")
+		TEXT("1 - Show an error \"Slashed Circle\"\n")
+	);
+	
 	TAutoConsoleVariable<int32> CVarEnableITFSoftwareCursor(
 		TEXT("Editor.ITF.SoftwareCursorOverride"),
 		2,
-		TEXT("0 - No software cursor will be used")
-		TEXT("1 - Software cursor will be used when requested")
-		TEXT("2 - Software cursor will be used when requested, defaulting to On when set to \"Grab Hand\""));
+		TEXT("0 - No software cursor will be used\n")
+		TEXT("1 - Software cursor will be used when requested\n")
+		TEXT("2 - Software cursor will be used when requested, defaulting to On when set to \"Grab Hand\"\n"));
 }
 
 class FEdModeToolsContextQueriesImpl : public IToolsContextQueriesAPI
@@ -1355,7 +1363,7 @@ bool UModeManagerInteractiveToolsContext::GetCursor(EMouseCursor::Type& OutCurso
 		return true;
 	}
 
-	if (bIsTrackingMouse)
+	if (bIsTrackingMouse && UE::Editor::Private::CVarCursorOverrideFallback.GetValueOnAnyThread() == 1)
 	{
 		OutCursor = EMouseCursor::SlashedCircle;				
 		return true;
