Allow newly created packages with a file on disk to be reloaded

#Synced-CL 52715155

diff --git a/Engine/Source/Editor/UnrealEd/Private/PackageTools.cpp b/Engine/Source/Editor/UnrealEd/Private/PackageTools.cpp
--- a/Engine/Source/Editor/UnrealEd/Private/PackageTools.cpp
+++ b/Engine/Source/Editor/UnrealEd/Private/PackageTools.cpp
@@ -661,7 +661,20 @@ UPackageTools::UPackageTools(const FObjectInitializer& ObjectInitializer)
 	{
 		void Add(UPackage* RealPackage)
 		{
-			if (RealPackage->HasAnyPackageFlags(PKG_InMemoryOnly))
+			if (RealPackage->HasAnyPackageFlags(PKG_NewlyCreated))
+			{
+				// Newly created packages *may* exist on disk (eg, deleting an actor, submitting, and then reverting the delete)
+				// so in that case we want to allow them to reload via PackagesToReload rather than block them via InMemoryPackages
+				if (FPackageName::DoesPackageExist(RealPackage->GetName()))
+				{
+					PackagesToReload.AddUnique(RealPackage);
+				}
+				else
+				{
+					InMemoryPackages.AddUnique(RealPackage);
+				}
+			}
+			else if (RealPackage->HasAnyPackageFlags(PKG_InMemoryOnly))
 			{
 				InMemoryPackages.AddUnique(RealPackage);
 			}
