Merging CL#52757557

UE: Bump incrementalcook version. Tags not getting picked up

diff --git a/Engine/Source/Editor/UnrealEd/Private/Cooker/CookIncrementalVersion.cpp b/Engine/Source/Editor/UnrealEd/Private/Cooker/CookIncrementalVersion.cpp
--- a/Engine/Source/Editor/UnrealEd/Private/Cooker/CookIncrementalVersion.cpp
+++ b/Engine/Source/Editor/UnrealEd/Private/Cooker/CookIncrementalVersion.cpp
@@ -8,6 +8,6 @@ namespace UE::Cook
 {
 
 // Change CookIncrementalVersion to a new guid when all packages in an incremental cook need to be invalidated
-FGuid CookIncrementalVersion(0x047D6737, 0x9B2949E3, 0xBF0493AF, 0x2C75D282);
+FGuid CookIncrementalVersion(0xFFA412DC, 0x88224FE2, 0xBEF7627A, 0x1B71D307);
 
 }
