Engine: Fix PIE level streaming

#rnx

#Synced-CL 52734363

diff --git a/Engine/Source/Editor/UnrealEd/Private/EditorEngine.cpp b/Engine/Source/Editor/UnrealEd/Private/EditorEngine.cpp
--- a/Engine/Source/Editor/UnrealEd/Private/EditorEngine.cpp
+++ b/Engine/Source/Editor/UnrealEd/Private/EditorEngine.cpp
@@ -2395,35 +2395,47 @@ void UEditorEngine::Tick( float DeltaSeconds, bool bIdleMode )
 
 	if( FSlateThrottleManager::Get().IsAllowingExpensiveTasks() )
 	{
-		SCOPE_CYCLE_COUNTER(STAT_RedrawViewports);
-		CSV_SCOPED_TIMING_STAT_EXCLUSIVE(ViewportMisc);
-		for (auto ContextIt = WorldList.CreateIterator(); ContextIt; ++ContextIt)
 		{
-			FWorldContext &PieContext = *ContextIt;
-			if (PieContext.WorldType != EWorldType::PIE)
+			SCOPE_CYCLE_COUNTER(STAT_RedrawViewports);
+			CSV_SCOPED_TIMING_STAT_EXCLUSIVE(ViewportMisc);
+			for (auto ContextIt = WorldList.CreateIterator(); ContextIt; ++ContextIt)
 			{
-				continue;
-			}
+				FWorldContext &PieContext = *ContextIt;
+				if (PieContext.WorldType != EWorldType::PIE)
+				{
+					continue;
+				}
 
-			PlayWorld = PieContext.World();
-			GameViewport = PieContext.GameViewport;
+				PlayWorld = PieContext.World();
+				GameViewport = PieContext.GameViewport;
 
-			// Render playworld. This needs to happen after the other viewports for screenshots to work correctly in PIE.
-			if (PlayWorld && GameViewport && !bIsSimulatingInEditor)
-			{
-				// Use the PlayWorld as the GWorld, because who knows what will happen in the Tick.
-				UWorld* OldGWorld = SetPlayInEditorWorld( PlayWorld );
+				// Render playworld. This needs to happen after the other viewports for screenshots to work correctly in PIE.
+				if (PlayWorld && GameViewport && !bIsSimulatingInEditor)
+				{
+					// Use the PlayWorld as the GWorld, because who knows what will happen in the Tick.
+					UWorld* OldGWorld = SetPlayInEditorWorld( PlayWorld );
 
-				// Render everything.
-				GameViewport->LayoutPlayers();
-				check(GameViewport->Viewport);
-				GameViewport->Viewport->Draw();
+					// Render everything.
+					GameViewport->LayoutPlayers();
+					check(GameViewport->Viewport);
+					GameViewport->Viewport->Draw();
 
-				// Pop the world
-				RestoreEditorWorld( OldGWorld );
+					// Pop the world
+					RestoreEditorWorld( OldGWorld );
+				}
+			}
+			RedrawViewports();
+		}
+		{
+			SCOPE_CYCLE_COUNTER(STAT_UpdateLevelStreaming);
+			for (FWorldContext& Context : WorldList)
+			{
+				if (Context.World() && Context.WorldType == EWorldType::PIE && !Context.RunAsDedicated && !GUseUnifiedTimeBudgetForStreaming)
+				{
+					Context.World()->UpdateLevelStreaming();
+				}
 			}
 		}
-		RedrawViewports();
 	}
 
 	// Update resource streaming after both regular Editor viewports and PIE had a chance to add viewers.
