[EditorEngine]
- Avoid throttling the editor during loading and async compilation

#Synced-CL 52786797

diff --git a/Engine/Source/Editor/UnrealEd/Private/EditorEngine.cpp b/Engine/Source/Editor/UnrealEd/Private/EditorEngine.cpp
--- a/Engine/Source/Editor/UnrealEd/Private/EditorEngine.cpp
+++ b/Engine/Source/Editor/UnrealEd/Private/EditorEngine.cpp
@@ -5167,6 +5167,16 @@ bool UEditorEngine::ShouldThrottleCPUUsage() const
 				}
 			}
 		}
+
+		if (bShouldThrottle)
+		{
+			// postloading and asset compilation needs to tick on GT and throttling them
+			// can make editor loading a lot longer.
+			if (IsAsyncLoading() || FAssetCompilingManager::Get().GetNumRemainingAssets() > 0)
+			{
+				bShouldThrottle = false;
+			}
+		}
 	}
 
 	return bShouldThrottle;
